Renamed Cascade ToggleRealtime tooltip to match the Editor viewports

#ttp 333918 - Editor: L10N: Conflicting localization resources log spew

#branch UE4

diff --git a/Engine/Source/Editor/Cascade/Private/CascadeActions.cpp b/Engine/Source/Editor/Cascade/Private/CascadeActions.cpp
--- a/Engine/Source/Editor/Cascade/Private/CascadeActions.cpp
+++ b/Engine/Source/Editor/Cascade/Private/CascadeActions.cpp
@@ -50,7 +50,7 @@ void FCascadeCommands::RegisterCommands()
 	UI_COMMAND(TogglePostProcess, "Post Process", "Toggle Post Process", EUserInterfaceActionType::ToggleButton, FInputGesture());
 	UI_COMMAND(ToggleGrid, "Grid", "Display Grid", EUserInterfaceActionType::ToggleButton, FInputGesture());
 	UI_COMMAND(ToggleLoopSystem, "Loop", "Toggle Looping", EUserInterfaceActionType::ToggleButton, FInputGesture());
-	UI_COMMAND(ToggleRealtime, "Realtime", "Toggle Realtime", EUserInterfaceActionType::ToggleButton, FInputGesture());
+	UI_COMMAND(ToggleRealtime, "Realtime", "Toggles real time rendering in this viewport", EUserInterfaceActionType::ToggleButton, FInputGesture());
 	UI_COMMAND(CascadeBackgroundColor, "Background Color", "Change Background Color", EUserInterfaceActionType::Button, FInputGesture());
 	UI_COMMAND(ToggleWireframeSphere, "Wireframe Sphere", "Display Wireframe Sphere", EUserInterfaceActionType::ToggleButton, FInputGesture());
 
