New emitter render mode "Lights Only": Particles are simulated and spawn lights, but sprites are not rendered; TTP 338195

diff --git a/Engine/Source/Editor/Cascade/Private/CascadeEmitterCanvasClient.cpp b/Engine/Source/Editor/Cascade/Private/CascadeEmitterCanvasClient.cpp
--- a/Engine/Source/Editor/Cascade/Private/CascadeEmitterCanvasClient.cpp
+++ b/Engine/Source/Editor/Cascade/Private/CascadeEmitterCanvasClient.cpp
@@ -120,8 +120,10 @@ FCascadeEmitterCanvasClient::FCascadeEmitterCanvasClient(TWeakPtr<FCascade> InCa
 	check(IconTex[Icon_RenderPoint]);
 	IconTex[Icon_RenderNone] = (UTexture2D*)StaticLoadObject(UTexture2D::StaticClass(), NULL, TEXT("/Engine/EditorMaterials/Cascade/CASC_None.CASC_None"), NULL, LOAD_None, NULL);
 	check(IconTex[Icon_RenderNone]);
+	IconTex[Icon_RenderLights] = (UTexture2D*)StaticLoadObject(UTexture2D::StaticClass(), NULL, TEXT("/Engine/EditorMaterials/Cascade/CASC_Lights.CASC_Lights"), NULL, LOAD_None, NULL);
+	check(IconTex[Icon_RenderLights]);
 	IconTex[Icon_CurveEdit] = (UTexture2D*)StaticLoadObject(UTexture2D::StaticClass(), NULL, TEXT("/Engine/EditorMaterials/Cascade/CASC_CurveEd.CASC_CurveEd"), NULL, LOAD_None, NULL);
-	check(IconTex[Icon_CurveEdit]);
+	check(IconTex[Icon_RenderLights]);
 	IconTex[Icon_3DDrawEnabled] = (UTexture2D*)StaticLoadObject(UTexture2D::StaticClass(), NULL, TEXT("/Engine/EditorMaterials/Cascade/CASC_ModuleEnable.CASC_ModuleEnable"), NULL, LOAD_None, NULL);
 	check(IconTex[Icon_3DDrawEnabled]);
 	IconTex[Icon_3DDrawDisabled] = (UTexture2D*)StaticLoadObject(UTexture2D::StaticClass(), NULL, TEXT("/Engine/EditorMaterials/Cascade/CASC_ModuleDisable.CASC_ModuleDisable"), NULL, LOAD_None, NULL);
@@ -305,8 +307,9 @@ bool FCascadeEmitterCanvasClient::InputKey(FViewport* Viewport, int32 Controller
 						{
 						case ERM_Normal:	DrawMode	= ERM_Point;	break;
 						case ERM_Point:		DrawMode	= ERM_Cross;	break;
-						case ERM_Cross:		DrawMode	= ERM_None;		break;
-						case ERM_None:		DrawMode	= ERM_Normal;	break;
+						case ERM_Cross:		DrawMode	= ERM_LightsOnly;		break;
+						case ERM_LightsOnly:		DrawMode = ERM_None;		break;
+						case ERM_None:		DrawMode = ERM_Normal;	break;
 						}
 						CascadePtr.Pin()->SetSelectedEmitter(Emitter);
 						UParticleLODLevel* LODLevel = CascadePtr.Pin()->GetCurrentlySelectedLODLevel(Emitter);
@@ -1142,6 +1145,9 @@ void FCascadeEmitterCanvasClient::DrawHeaderBlock(int32 Index, int32 XPos, UPart
 	case ERM_Cross:
 		IconTxtr	= GetIconTexture(Icon_RenderCross);
 		break;
+	case ERM_LightsOnly:
+		IconTxtr = GetIconTexture(Icon_RenderLights);
+		break;
 	case ERM_None:
 		IconTxtr	= GetIconTexture(Icon_RenderNone);
 		break;
diff --git a/Engine/Source/Editor/Cascade/Private/CascadeEmitterCanvasClient.h b/Engine/Source/Editor/Cascade/Private/CascadeEmitterCanvasClient.h
--- a/Engine/Source/Editor/Cascade/Private/CascadeEmitterCanvasClient.h
+++ b/Engine/Source/Editor/Cascade/Private/CascadeEmitterCanvasClient.h
@@ -48,6 +48,7 @@ private:
 		Icon_RenderCross,
 		Icon_RenderPoint,
 		Icon_RenderNone,
+		Icon_RenderLights,
 		Icon_CurveEdit,
 		Icon_3DDrawEnabled,
 		Icon_3DDrawDisabled,
diff --git a/Engine/Source/Runtime/Engine/Classes/Particles/ParticleEmitter.h b/Engine/Source/Runtime/Engine/Classes/Particles/ParticleEmitter.h
--- a/Engine/Source/Runtime/Engine/Classes/Particles/ParticleEmitter.h
+++ b/Engine/Source/Runtime/Engine/Classes/Particles/ParticleEmitter.h
@@ -44,7 +44,8 @@ enum EEmitterRenderMode
 	ERM_Normal UMETA(DisplayName="Normal"),
 	ERM_Point UMETA(DisplayName="Point"),
 	ERM_Cross UMETA(DisplayName="Cross"),
-	ERM_None UMETA(DisplayName="None"),
+	ERM_LightsOnly UMETA(DisplayName = "Lights Only"),
+	ERM_None UMETA(DisplayName = "None"),
 	ERM_MAX,
 };
 
