Particle effect meshes now adhere to the view mode correctly

#ttp 340954 - CASCADE: Switching view mode from wireframe will keep mesh in wireframe mode

#branch UE4

#change Particle system component is re-added to the scene so that it picks up the correct viewmode

diff --git a/Engine/Source/Editor/Cascade/Private/Cascade.cpp b/Engine/Source/Editor/Cascade/Private/Cascade.cpp
--- a/Engine/Source/Editor/Cascade/Private/Cascade.cpp
+++ b/Engine/Source/Editor/Cascade/Private/Cascade.cpp
@@ -2489,6 +2489,18 @@ void FCascade::SetLODValue(int32 LODSetting)
 	}
 }
 
+void FCascade::ReassociateParticleSystem() const
+{
+	if (ParticleSystemComponent)
+	{
+		if (PreviewViewport.IsValid() && PreviewViewport->GetViewportClient().IsValid())
+		{
+			PreviewViewport->GetViewportClient()->GetPreviewScene().RemoveComponent(ParticleSystemComponent);
+			PreviewViewport->GetViewportClient()->GetPreviewScene().AddComponent(ParticleSystemComponent, FTransform::Identity);
+		}
+	}
+}
+
 void FCascade::RestartParticleSystem()
 {
 	if (ParticleSystemComponent)
@@ -2502,11 +2514,8 @@ void FCascade::RestartParticleSystem()
 		ParticleSystemComponent->bIsViewRelevanceDirty = true;
 		ParticleSystemComponent->CachedViewRelevanceFlags.Empty();
 		ParticleSystemComponent->ConditionalCacheViewRelevanceFlags();
-		if (PreviewViewport.IsValid() && PreviewViewport->GetViewportClient().IsValid())
-		{
-			PreviewViewport->GetViewportClient()->GetPreviewScene().RemoveComponent(ParticleSystemComponent);
-			PreviewViewport->GetViewportClient()->GetPreviewScene().AddComponent(ParticleSystemComponent, FTransform::Identity);
-		}
+
+		ReassociateParticleSystem();
 	}
 
 	if (ParticleSystem)
@@ -3645,6 +3654,9 @@ void FCascade::OnViewMode(EViewModeIndex ViewMode)
 	if (PreviewViewport.IsValid() && PreviewViewport->GetViewportClient().IsValid())
 	{
 		PreviewViewport->GetViewportClient()->SetViewMode(ViewMode);
+
+		ReassociateParticleSystem();
+
 		PreviewViewport->RefreshViewport();
 	}
 }
@@ -3698,11 +3710,7 @@ void FCascade::OnToggleBoundsSetFixedBounds()
 		SetSelection(NewSelection);
 	}
 
-	if (PreviewViewport.IsValid() && PreviewViewport->GetViewportClient().IsValid())
-	{
-		PreviewViewport->GetViewportClient()->GetPreviewScene().RemoveComponent(ParticleSystemComponent);
-		PreviewViewport->GetViewportClient()->GetPreviewScene().AddComponent(ParticleSystemComponent, FTransform::Identity);
-	}
+	ReassociateParticleSystem();
 }
 
 void FCascade::OnTogglePostProcess()
diff --git a/Engine/Source/Editor/Cascade/Private/Cascade.h b/Engine/Source/Editor/Cascade/Private/Cascade.h
--- a/Engine/Source/Editor/Cascade/Private/Cascade.h
+++ b/Engine/Source/Editor/Cascade/Private/Cascade.h
@@ -243,6 +243,9 @@ private:
 	/** Changes the LOD level */
 	void SetLODValue(int32 LODSetting);
 
+	/** Removes and then Readds the particle system to the preview scene */
+	void ReassociateParticleSystem() const;
+
 	/** Restarts the particle system */
 	void RestartParticleSystem();
 
