Added analytics reporting to cascade

#ttp 340620 - FIXIF: EDITOR:ANALYTICS: Cascade Editor Instrumentation Pass

#branch UE4

#added Added analytics for various cascade usage; whenever a particle effect is opens it gives a breif overview of how many emitters/lods/modules there are, and also reports when any emitter/lod/modules gets added/removed

diff --git a/Engine/Source/Editor/Cascade/Private/Cascade.cpp b/Engine/Source/Editor/Cascade/Private/Cascade.cpp
--- a/Engine/Source/Editor/Cascade/Private/Cascade.cpp
+++ b/Engine/Source/Editor/Cascade/Private/Cascade.cpp
@@ -32,6 +32,9 @@
 #include "Particles/ParticleSystem.h"
 #include "Particles/ParticleSystemComponent.h"
 
+#include "Runtime/Analytics/Analytics/Public/Interfaces/IAnalyticsProvider.h"
+#include "EngineAnalytics.h"
+
 static const FName Cascade_PreviewViewportTab("Cascade_PreviewViewport");
 static const FName Cascade_EmmitterCanvasTab("Cascade_EmitterCanvas");
 static const FName Cascade_PropertiesTab("Cascade_Properties");
@@ -166,31 +169,52 @@ void FCascade::InitCascade(const EToolkitMode::Type Mode, const TSharedPtr< clas
 	EditorConfig = ConstructObject<UCascadeConfiguration>(UCascadeConfiguration::StaticClass());
 	check(EditorConfig);
 
+	FString Description;
 	for (int32 EmitterIdx = 0; EmitterIdx < ParticleSystem->Emitters.Num(); EmitterIdx++)
 	{
 		UParticleEmitter* Emitter = ParticleSystem->Emitters[EmitterIdx];
 		if (Emitter)
 		{
+			Description += FString::Printf(TEXT("Emitter%d["), EmitterIdx);
 			Emitter->SetFlags(RF_Transactional);
 			for (int32 LODIndex = 0; LODIndex < Emitter->LODLevels.Num(); LODIndex++)
 			{
 				UParticleLODLevel* LODLevel = Emitter->GetLODLevel(LODIndex);
 				if (LODLevel)
 				{
+					Description += FString::Printf(TEXT("LOD%d("), LODIndex);
 					LODLevel->SetFlags(RF_Transactional);
 					check(LODLevel->RequiredModule);
 					LODLevel->RequiredModule->SetTransactionFlag();
 					check(LODLevel->SpawnModule);
 					LODLevel->SpawnModule->SetTransactionFlag();
-					for (int32 ModuleIdx = 0; ModuleIdx < LODLevel->Modules.Num(); ModuleIdx++)
+					if (LODLevel->Modules.Num() > 0)
 					{
-						UParticleModule* pkModule = LODLevel->Modules[ModuleIdx];
-						pkModule->SetTransactionFlag();
+						Description += FString::Printf(TEXT("Modules%d"), LODLevel->Modules.Num());
+						for (int32 ModuleIdx = 0; ModuleIdx < LODLevel->Modules.Num(); ModuleIdx++)
+						{
+							UParticleModule* pkModule = LODLevel->Modules[ModuleIdx];
+							pkModule->SetTransactionFlag();
+						}
+					}
+					Description += TEXT(")");
+					if (Emitter->LODLevels.Num() > LODIndex + 1)
+					{
+						Description += TEXT(",");
 					}
 				}
 			}
+			Description += TEXT("]");
+			if (ParticleSystem->Emitters.Num() > EmitterIdx + 1)
+			{
+				Description += TEXT(",");
+			}
 		}
 	}
+	if (Description.Len() > 0 && FEngineAnalytics::IsAvailable())
+	{
+		FEngineAnalytics::GetProvider().RecordEvent(TEXT("Editor.Cascade.Init"), TEXT("Overview"), Description);
+	}
 
 	ParticleSystemComponent = ConstructObject<UCascadeParticleSystemComponent>(UCascadeParticleSystemComponent::StaticClass());
 
@@ -918,6 +942,11 @@ void FCascade::OnNewModule(int32 Idx)
 
 	EndTransaction( Transaction );
 
+	if (FEngineAnalytics::IsAvailable())
+	{
+		FEngineAnalytics::GetProvider().RecordEvent(TEXT("Editor.Cascade.NewModule"), TEXT("Class"), NewModClass->GetName());
+	}
+
 	ParticleSystem->MarkPackageDirty();
 
 	// Refresh viewport
@@ -1080,6 +1109,11 @@ void FCascade::OnNewEmitter()
 
 	EndTransaction( Transaction );
 
+	if (FEngineAnalytics::IsAvailable())
+	{
+		FEngineAnalytics::GetProvider().RecordEvent(TEXT("Editor.Cascade.NewEmitter"));
+	}
+
 	// Refresh viewport
 	if (EmitterCanvas.IsValid())
 	{
@@ -3301,6 +3335,11 @@ void FCascade::AddLOD(bool bBeforeCurrent)
 		check(bTransactionInProgress);
 		EndTransaction( Transaction );
 
+		if (FEngineAnalytics::IsAvailable())
+		{
+			FEngineAnalytics::GetProvider().RecordEvent(TEXT("Editor.Cascade.NewLOD"), FAnalyticsEventAttribute(TEXT("Index"), CurrentLODIndex));
+		}
+
 		UpdateLODLevel();
 		SetSelectedModule(SelectedEmitter, SelectedModule);
 		ForceUpdate();
@@ -4189,6 +4228,11 @@ void FCascade::OnDeleteLOD()
 	check(bTransactionInProgress);
 	EndTransaction( Transaction );
 
+	if (FEngineAnalytics::IsAvailable())
+	{
+		FEngineAnalytics::GetProvider().RecordEvent(TEXT("Editor.Cascade.DeleteLOD"), FAnalyticsEventAttribute(TEXT("Index"), Selection));
+	}
+
 	ForceUpdate();
 
 	OnRestartInLevel();
@@ -4274,6 +4318,7 @@ void FCascade::OnDeleteModule(bool bConfirm)
 
 	// Find the module index...
 	int32	DeleteModuleIndex	= -1;
+	FString	ModuleName;
 
 	UParticleLODLevel* HighLODLevel	= SelectedEmitter->GetLODLevel(0);
 	check(HighLODLevel);
@@ -4283,6 +4328,7 @@ void FCascade::OnDeleteModule(bool bConfirm)
 		if (CheckModule == SelectedModule)
 		{
 			DeleteModuleIndex = ModuleIndex;
+			ModuleName = CheckModule->GetClass()->GetName();
 			break;
 		}
 	}
@@ -4373,6 +4419,11 @@ void FCascade::OnDeleteModule(bool bConfirm)
 
 	EndTransaction( Transaction );
 
+	if (FEngineAnalytics::IsAvailable())
+	{
+		FEngineAnalytics::GetProvider().RecordEvent(TEXT("Editor.Cascade.DeleteModule"), TEXT("Class"), ModuleName);
+	}
+
 	SetSelectedEmitter(SelectedEmitter);
 
 	if (EmitterCanvas.IsValid())
@@ -4885,6 +4936,11 @@ void FCascade::OnDeleteEmitter()
 
 	EndTransaction( Transaction );
 
+	if (FEngineAnalytics::IsAvailable())
+	{
+		FEngineAnalytics::GetProvider().RecordEvent(TEXT("Editor.Cascade.DeleteEmitter"));
+	}
+
 	ParticleSystem->MarkPackageDirty();
 
 	if (EmitterCanvas.IsValid())
