Fixing crash in curve editing; Cascade crashes if keyboard focus is in an edit field in the details pane when selecting and dragging a curve point. Now clearing keyboard focus explicitly on curve edit, which prevents this crash, but requires two clicks into the curve editor (one to focus, one to start dragging) for curve editing.

diff --git a/Engine/Source/Editor/Cascade/Private/Cascade.cpp b/Engine/Source/Editor/Cascade/Private/Cascade.cpp
--- a/Engine/Source/Editor/Cascade/Private/Cascade.cpp
+++ b/Engine/Source/Editor/Cascade/Private/Cascade.cpp
@@ -1630,7 +1630,7 @@ void FCascade::NotifyPreChange(FEditPropertyChain* PropertyChain)
 	//Needs to stay in-sync with the text in NotifyPostChange
 	FText Transaction = NSLOCTEXT("UnrealEd", "CascadePropertyChange", "Change Property");
 
-	BeginTransaction( Transaction );
+	BeginTransaction(Transaction);
 	ModifyParticleSystem();
 
 	CurveToReplace = NULL;
@@ -1696,6 +1696,7 @@ void FCascade::NotifyPreChange(FEditPropertyChain* PropertyChain)
 			}
 		}
 	}
+
 }
 
 void FCascade::NotifyPostChange(const FPropertyChangedEvent& PropertyChangedEvent, FEditPropertyChain* PropertyChain)
@@ -1884,6 +1885,8 @@ void FCascade::NotifyPostChange(const FPropertyChangedEvent& PropertyChangedEven
 
 void FCascade::PreEditCurve(TArray<UObject*> CurvesAboutToChange)
 {
+	FSlateApplication::Get().ClearKeyboardFocus(EKeyboardFocusCause::Mouse);
+
 	//Need to keep text in-sync with PostEditCurve
 	BeginTransaction( NSLOCTEXT("UnrealEd", "EditCurve", "Edit Curve") );
 	ModifyParticleSystem();
