Prepare cpp files for header cleanup.

diff --git a/Engine/Source/Editor/ActorPickerMode/Private/EditorModeActorPicker.cpp b/Engine/Source/Editor/ActorPickerMode/Private/EditorModeActorPicker.cpp
--- a/Engine/Source/Editor/ActorPickerMode/Private/EditorModeActorPicker.cpp
+++ b/Engine/Source/Editor/ActorPickerMode/Private/EditorModeActorPicker.cpp
@@ -1,6 +1,7 @@
 // Copyright 1998-2014 Epic Games, Inc. All Rights Reserved.
 
 #include "ActorPickerModePrivatePCH.h"
+#include "EngineUtils.h"
 
 #define LOCTEXT_NAMESPACE "PropertyPicker"
 
diff --git a/Engine/Source/Editor/AnimGraph/Classes/AnimGraphNode_BlendSpaceBase.h b/Engine/Source/Editor/AnimGraph/Classes/AnimGraphNode_BlendSpaceBase.h
--- a/Engine/Source/Editor/AnimGraph/Classes/AnimGraphNode_BlendSpaceBase.h
+++ b/Engine/Source/Editor/AnimGraph/Classes/AnimGraphNode_BlendSpaceBase.h
@@ -2,6 +2,7 @@
 
 #pragma once
 #include "AnimGraphNode_Base.h"
+#include "Animation/BlendSpaceBase.h"
 #include "AnimGraphNode_BlendSpaceBase.generated.h"
 
 UCLASS(Abstract, MinimalAPI)
diff --git a/Engine/Source/Editor/AnimGraph/Classes/AnimStateTransitionNode.h b/Engine/Source/Editor/AnimGraph/Classes/AnimStateTransitionNode.h
--- a/Engine/Source/Editor/AnimGraph/Classes/AnimStateTransitionNode.h
+++ b/Engine/Source/Editor/AnimGraph/Classes/AnimStateTransitionNode.h
@@ -4,6 +4,7 @@
 
 #include "Animation/AnimInstance.h"
 #include "AnimStateNodeBase.h"
+#include "Animation/AnimStateMachineTypes.h"
 
 #include "AnimStateTransitionNode.generated.h"
 
diff --git a/Engine/Source/Editor/BehaviorTreeEditor/Private/BehaviorTreeEditor.h b/Engine/Source/Editor/BehaviorTreeEditor/Private/BehaviorTreeEditor.h
--- a/Engine/Source/Editor/BehaviorTreeEditor/Private/BehaviorTreeEditor.h
+++ b/Engine/Source/Editor/BehaviorTreeEditor/Private/BehaviorTreeEditor.h
@@ -4,6 +4,7 @@
 
 #include "IBehaviorTreeEditor.h"
 #include "Toolkits/AssetEditorToolkit.h"
+#include "EditorUndoClient.h"
 
 class UBehaviorTree;
 class UBlackboardData;
diff --git a/Engine/Source/Editor/BlueprintGraph/Classes/K2Node.h b/Engine/Source/Editor/BlueprintGraph/Classes/K2Node.h
--- a/Engine/Source/Editor/BlueprintGraph/Classes/K2Node.h
+++ b/Engine/Source/Editor/BlueprintGraph/Classes/K2Node.h
@@ -3,10 +3,12 @@
 #pragma once
 #include "EdGraph/EdGraphNode.h"
 #include "BlueprintNodeSignature.h"
+#include "EngineLogs.h"
 #include "K2Node.generated.h"
 
 class UBlueprintNodeSpawner;
 class FBlueprintActionDatabaseRegistrar;
+class UDynamicBlueprintBinding;
 
 /** Helper struct to allow us to redirect properties and functions through renames and additionally between classes if necessary */
 struct FFieldRemapInfo
diff --git a/Engine/Source/Editor/BlueprintGraph/Classes/K2Node_ComponentBoundEvent.h b/Engine/Source/Editor/BlueprintGraph/Classes/K2Node_ComponentBoundEvent.h
--- a/Engine/Source/Editor/BlueprintGraph/Classes/K2Node_ComponentBoundEvent.h
+++ b/Engine/Source/Editor/BlueprintGraph/Classes/K2Node_ComponentBoundEvent.h
@@ -6,6 +6,8 @@
 #include "EdGraph/EdGraphNodeUtils.h" // for FNodeTextCache
 #include "K2Node_ComponentBoundEvent.generated.h"
 
+class UDynamicBlueprintBinding;
+
 UCLASS(MinimalAPI)
 class UK2Node_ComponentBoundEvent : public UK2Node_Event
 {
diff --git a/Engine/Source/Editor/BlueprintGraph/Classes/K2Node_GetDataTableRow.h b/Engine/Source/Editor/BlueprintGraph/Classes/K2Node_GetDataTableRow.h
--- a/Engine/Source/Editor/BlueprintGraph/Classes/K2Node_GetDataTableRow.h
+++ b/Engine/Source/Editor/BlueprintGraph/Classes/K2Node_GetDataTableRow.h
@@ -4,6 +4,7 @@
 #pragma once
 #include "K2Node.h"
 #include "EdGraph/EdGraphNodeUtils.h" // for FNodeTextCache
+#include "Engine/DataTable.h"
 #include "K2Node_GetDataTableRow.generated.h"
 
 UCLASS()
diff --git a/Engine/Source/Editor/BlueprintGraph/Classes/K2Node_InputKeyEvent.h b/Engine/Source/Editor/BlueprintGraph/Classes/K2Node_InputKeyEvent.h
--- a/Engine/Source/Editor/BlueprintGraph/Classes/K2Node_InputKeyEvent.h
+++ b/Engine/Source/Editor/BlueprintGraph/Classes/K2Node_InputKeyEvent.h
@@ -2,6 +2,7 @@
 
 
 #pragma once
+#include "Components/InputComponent.h"
 #include "K2Node_Event.h"
 #include "K2Node_InputKeyEvent.generated.h"
 
diff --git a/Engine/Source/Editor/Blutility/Private/AssetTypeActions_EditorUtilityBlueprint.cpp b/Engine/Source/Editor/Blutility/Private/AssetTypeActions_EditorUtilityBlueprint.cpp
--- a/Engine/Source/Editor/Blutility/Private/AssetTypeActions_EditorUtilityBlueprint.cpp
+++ b/Engine/Source/Editor/Blutility/Private/AssetTypeActions_EditorUtilityBlueprint.cpp
@@ -9,6 +9,7 @@
 #include "AssetRegistryModule.h"
 #include "AssetToolsModule.h"
 #include "GlobalBlutilityDialog.h"
+#include "Engine/BlueprintGeneratedClass.h"
 
 #define LOCTEXT_NAMESPACE "AssetTypeActions"
 
diff --git a/Engine/Source/Editor/Blutility/Private/EditorUtilityBlueprintFactory.cpp b/Engine/Source/Editor/Blutility/Private/EditorUtilityBlueprintFactory.cpp
--- a/Engine/Source/Editor/Blutility/Private/EditorUtilityBlueprintFactory.cpp
+++ b/Engine/Source/Editor/Blutility/Private/EditorUtilityBlueprintFactory.cpp
@@ -6,6 +6,7 @@
 #include "BlutilityClasses.h"
 #include "Kismet2/KismetEditorUtilities.h"
 #include "Kismet2/SClassPickerDialog.h"
+#include "Engine/BlueprintGeneratedClass.h"
 
 class FBlutilityBlueprintFactoryFilter : public IClassViewerFilter
 {
diff --git a/Engine/Source/Editor/Blutility/Private/GlobalEditorUtilityBase.cpp b/Engine/Source/Editor/Blutility/Private/GlobalEditorUtilityBase.cpp
--- a/Engine/Source/Editor/Blutility/Private/GlobalEditorUtilityBase.cpp
+++ b/Engine/Source/Editor/Blutility/Private/GlobalEditorUtilityBase.cpp
@@ -6,6 +6,7 @@
 #include "ContentBrowserModule.h"
 #include "AssetData.h"
 #include "AssetToolsModule.h"
+#include "Engine/Selection.h"
 
 /////////////////////////////////////////////////////
 
diff --git a/Engine/Source/Editor/Cascade/Private/Cascade.cpp b/Engine/Source/Editor/Cascade/Private/Cascade.cpp
--- a/Engine/Source/Editor/Cascade/Private/Cascade.cpp
+++ b/Engine/Source/Editor/Cascade/Private/Cascade.cpp
@@ -41,6 +41,15 @@
 #include "STextEntryPopup.h"
 #include "GenericCommands.h"
 #include "Components/VectorFieldComponent.h"
+#include "Engine/StaticMesh.h"
+#include "Engine/InterpCurveEdSetup.h"
+#include "Engine/Selection.h"
+#include "Particles/ParticleModule.h"
+#include "Distributions/DistributionFloatUniform.h"
+#include "Distributions/DistributionFloatUniformCurve.h"
+#include "Distributions/DistributionVectorUniform.h"
+#include "Distributions/DistributionVectorUniformCurve.h"
+#include "UnrealEngine.h"
 
 static const FName Cascade_PreviewViewportTab("Cascade_PreviewViewport");
 static const FName Cascade_EmmitterCanvasTab("Cascade_EmitterCanvas");
diff --git a/Engine/Source/Editor/Cascade/Private/Cascade.h b/Engine/Source/Editor/Cascade/Private/Cascade.h
--- a/Engine/Source/Editor/Cascade/Private/Cascade.h
+++ b/Engine/Source/Editor/Cascade/Private/Cascade.h
@@ -8,6 +8,7 @@
 #include "CascadeParticleSystemComponent.h"
 #include "Particles/ParticleEmitter.h"
 #include "Particles/ParticleModule.h"
+#include "EditorUndoClient.h"
 
 class SCascadePreviewViewport;
 class SCascadeEmitterCanvas;
diff --git a/Engine/Source/Editor/Cascade/Private/CascadeEmitterCanvasClient.cpp b/Engine/Source/Editor/Cascade/Private/CascadeEmitterCanvasClient.cpp
--- a/Engine/Source/Editor/Cascade/Private/CascadeEmitterCanvasClient.cpp
+++ b/Engine/Source/Editor/Cascade/Private/CascadeEmitterCanvasClient.cpp
@@ -17,6 +17,9 @@
 #include "Particles/ParticleSpriteEmitter.h"
 #include "Particles/ParticleSystem.h"
 #include "SColorPicker.h"
+#include "CanvasTypes.h"
+#include "Engine/Font.h"
+#include "Engine/StaticMesh.h"
 
 
 FCascadeEmitterCanvasClient::FCascadeEmitterCanvasClient(TWeakPtr<FCascade> InCascade, TWeakPtr<SCascadeEmitterCanvas> InCascadeViewport)
diff --git a/Engine/Source/Editor/Cascade/Private/CascadePreviewViewportClient.cpp b/Engine/Source/Editor/Cascade/Private/CascadePreviewViewportClient.cpp
--- a/Engine/Source/Editor/Cascade/Private/CascadePreviewViewportClient.cpp
+++ b/Engine/Source/Editor/Cascade/Private/CascadePreviewViewportClient.cpp
@@ -14,6 +14,10 @@
 #include "Particles/ParticleSpriteEmitter.h"
 #include "PhysicsPublic.h"
 #include "Components/VectorFieldComponent.h"
+#include "Engine/StaticMesh.h"
+#include "CanvasTypes.h"
+#include "Components/LineBatchComponent.h"
+#include "Engine/Canvas.h"
 
 #define LOCTEXT_NAMESPACE "CascadeViewportClient"
 
diff --git a/Engine/Source/Editor/ClassViewer/Private/UnloadedBlueprintData.cpp b/Engine/Source/Editor/ClassViewer/Private/UnloadedBlueprintData.cpp
--- a/Engine/Source/Editor/ClassViewer/Private/UnloadedBlueprintData.cpp
+++ b/Engine/Source/Editor/ClassViewer/Private/UnloadedBlueprintData.cpp
@@ -4,6 +4,7 @@
 
 #include "UnloadedBlueprintData.h"
 #include "ClassViewerNode.h"
+#include "Engine/BlueprintGeneratedClass.h"
 
 FUnloadedBlueprintData::FUnloadedBlueprintData( TWeakPtr< FClassViewerNode > InClassViewerNode )
 {
diff --git a/Engine/Source/Editor/ComponentVisualizers/Private/AudioComponentVisualizer.cpp b/Engine/Source/Editor/ComponentVisualizers/Private/AudioComponentVisualizer.cpp
--- a/Engine/Source/Editor/ComponentVisualizers/Private/AudioComponentVisualizer.cpp
+++ b/Engine/Source/Editor/ComponentVisualizers/Private/AudioComponentVisualizer.cpp
@@ -4,6 +4,7 @@
 
 #include "SoundDefinitions.h"
 #include "AudioComponentVisualizer.h"
+#include "Components/AudioComponent.h"
 
 
 void FAudioComponentVisualizer::DrawVisualization( const UActorComponent* Component, const FSceneView* View, FPrimitiveDrawInterface* PDI )
diff --git a/Engine/Source/Editor/ComponentVisualizers/Private/PointLightComponentVisualizer.cpp b/Engine/Source/Editor/ComponentVisualizers/Private/PointLightComponentVisualizer.cpp
--- a/Engine/Source/Editor/ComponentVisualizers/Private/PointLightComponentVisualizer.cpp
+++ b/Engine/Source/Editor/ComponentVisualizers/Private/PointLightComponentVisualizer.cpp
@@ -3,6 +3,7 @@
 #include "ComponentVisualizersPrivatePCH.h"
 
 #include "PointLightComponentVisualizer.h"
+#include "Components/PointLightComponent.h"
 
 
 void FPointLightComponentVisualizer::DrawVisualization( const UActorComponent* Component, const FSceneView* View, FPrimitiveDrawInterface* PDI )
diff --git a/Engine/Source/Editor/ComponentVisualizers/Private/RadialForceComponentVisualizer.cpp b/Engine/Source/Editor/ComponentVisualizers/Private/RadialForceComponentVisualizer.cpp
--- a/Engine/Source/Editor/ComponentVisualizers/Private/RadialForceComponentVisualizer.cpp
+++ b/Engine/Source/Editor/ComponentVisualizers/Private/RadialForceComponentVisualizer.cpp
@@ -3,6 +3,7 @@
 #include "ComponentVisualizersPrivatePCH.h"
 
 #include "RadialForceComponentVisualizer.h"
+#include "PhysicsEngine/RadialForceComponent.h"
 
 
 void FRadialForceComponentVisualizer::DrawVisualization( const UActorComponent* Component, const FSceneView* View, FPrimitiveDrawInterface* PDI )
diff --git a/Engine/Source/Editor/ComponentVisualizers/Private/SpotLightComponentVisualizer.cpp b/Engine/Source/Editor/ComponentVisualizers/Private/SpotLightComponentVisualizer.cpp
--- a/Engine/Source/Editor/ComponentVisualizers/Private/SpotLightComponentVisualizer.cpp
+++ b/Engine/Source/Editor/ComponentVisualizers/Private/SpotLightComponentVisualizer.cpp
@@ -3,6 +3,7 @@
 #include "ComponentVisualizersPrivatePCH.h"
 
 #include "SpotLightComponentVisualizer.h"
+#include "Components/SpotLightComponent.h"
 
 
 void FSpotLightComponentVisualizer::DrawVisualization( const UActorComponent* Component, const FSceneView* View, FPrimitiveDrawInterface* PDI )
diff --git a/Engine/Source/Editor/ComponentVisualizers/Private/SpringArmComponentVisualizer.cpp b/Engine/Source/Editor/ComponentVisualizers/Private/SpringArmComponentVisualizer.cpp
--- a/Engine/Source/Editor/ComponentVisualizers/Private/SpringArmComponentVisualizer.cpp
+++ b/Engine/Source/Editor/ComponentVisualizers/Private/SpringArmComponentVisualizer.cpp
@@ -2,6 +2,7 @@
 
 #include "ComponentVisualizersPrivatePCH.h"
 #include "SpringArmComponentVisualizer.h"
+#include "GameFramework/SpringArmComponent.h"
 
 static const FColor	ArmColor(255,0,0);
 
diff --git a/Engine/Source/Editor/ContentBrowser/Private/AssetContextMenu.cpp b/Engine/Source/Editor/ContentBrowser/Private/AssetContextMenu.cpp
--- a/Engine/Source/Editor/ContentBrowser/Private/AssetContextMenu.cpp
+++ b/Engine/Source/Editor/ContentBrowser/Private/AssetContextMenu.cpp
@@ -32,6 +32,7 @@
 #include "GenericCommands.h"
 #include "SNotificationList.h"
 #include "NotificationManager.h"
+#include "Engine/LevelStreaming.h"
 
 
 #define LOCTEXT_NAMESPACE "ContentBrowser"
diff --git a/Engine/Source/Editor/ContentBrowser/Private/SContentBrowser.cpp b/Engine/Source/Editor/ContentBrowser/Private/SContentBrowser.cpp
--- a/Engine/Source/Editor/ContentBrowser/Private/SContentBrowser.cpp
+++ b/Engine/Source/Editor/ContentBrowser/Private/SContentBrowser.cpp
@@ -12,6 +12,7 @@
 #include "SDockTab.h"
 #include "GenericCommands.h"
 #include "IAddContentDialogModule.h"
+#include "Engine/Selection.h"
 
 #define LOCTEXT_NAMESPACE "ContentBrowser"
 
diff --git a/Engine/Source/Editor/CurveAssetEditor/Private/CurveAssetEditor.cpp b/Engine/Source/Editor/CurveAssetEditor/Private/CurveAssetEditor.cpp
--- a/Engine/Source/Editor/CurveAssetEditor/Private/CurveAssetEditor.cpp
+++ b/Engine/Source/Editor/CurveAssetEditor/Private/CurveAssetEditor.cpp
@@ -10,6 +10,7 @@
 #include "SDockTab.h"
 #include "Curves/CurveVector.h"
 #include "Curves/CurveLinearColor.h"
+#include "Curves/CurveFloat.h"
 
 #define LOCTEXT_NAMESPACE "CurveAssetEditor"
 
diff --git a/Engine/Source/Editor/CurveTableEditor/Private/CurveTableEditor.cpp b/Engine/Source/Editor/CurveTableEditor/Private/CurveTableEditor.cpp
--- a/Engine/Source/Editor/CurveTableEditor/Private/CurveTableEditor.cpp
+++ b/Engine/Source/Editor/CurveTableEditor/Private/CurveTableEditor.cpp
@@ -6,6 +6,7 @@
 #include "Toolkits/IToolkitHost.h"
 #include "Editor/WorkspaceMenuStructure/Public/WorkspaceMenuStructureModule.h"
 #include "SDockTab.h"
+#include "Engine/CurveTable.h"
  
 #define LOCTEXT_NAMESPACE "CurveTableEditor"
 
diff --git a/Engine/Source/Editor/CurveTableEditor/Private/CurveTableEditor.h b/Engine/Source/Editor/CurveTableEditor/Private/CurveTableEditor.h
--- a/Engine/Source/Editor/CurveTableEditor/Private/CurveTableEditor.h
+++ b/Engine/Source/Editor/CurveTableEditor/Private/CurveTableEditor.h
@@ -5,6 +5,7 @@
 #include "ICurveTableEditor.h"
 #include "Toolkits/AssetEditorToolkit.h"
 
+class UCurveTable;
 
 /** Viewer/editor for a CurveTable */
 class FCurveTableEditor :
diff --git a/Engine/Source/Editor/CurveTableEditor/Private/CurveTableEditorModule.cpp b/Engine/Source/Editor/CurveTableEditor/Private/CurveTableEditorModule.cpp
--- a/Engine/Source/Editor/CurveTableEditor/Private/CurveTableEditorModule.cpp
+++ b/Engine/Source/Editor/CurveTableEditor/Private/CurveTableEditorModule.cpp
@@ -4,6 +4,7 @@
 #include "ModuleManager.h"
 #include "CurveTableEditor.h"
 #include "Toolkits/ToolkitManager.h"
+#include "Engine/CurveTable.h"
 
 IMPLEMENT_MODULE( FCurveTableEditorModule, CurveTableEditor );
 
diff --git a/Engine/Source/Editor/CurveTableEditor/Public/CurveTableEditorModule.h b/Engine/Source/Editor/CurveTableEditor/Public/CurveTableEditorModule.h
--- a/Engine/Source/Editor/CurveTableEditor/Public/CurveTableEditorModule.h
+++ b/Engine/Source/Editor/CurveTableEditor/Public/CurveTableEditorModule.h
@@ -7,6 +7,7 @@
 #include "Toolkits/AssetEditorToolkit.h" // For FExtensibilityManager
 
 class ICurveTableEditor;
+class UCurveTable;
 
 /** CurveTable Editor module */
 class FCurveTableEditorModule : public IModuleInterface,
diff --git a/Engine/Source/Editor/DataTableEditor/Private/DataTableEditor.cpp b/Engine/Source/Editor/DataTableEditor/Private/DataTableEditor.cpp
--- a/Engine/Source/Editor/DataTableEditor/Private/DataTableEditor.cpp
+++ b/Engine/Source/Editor/DataTableEditor/Private/DataTableEditor.cpp
@@ -7,6 +7,7 @@
 #include "Editor/WorkspaceMenuStructure/Public/WorkspaceMenuStructureModule.h"
 #include "SSearchBox.h"
 #include "SDockTab.h"
+#include "Engine/DataTable.h"
  
 #define LOCTEXT_NAMESPACE "DataTableEditor"
 
diff --git a/Engine/Source/Editor/DataTableEditor/Private/DataTableEditor.h b/Engine/Source/Editor/DataTableEditor/Private/DataTableEditor.h
--- a/Engine/Source/Editor/DataTableEditor/Private/DataTableEditor.h
+++ b/Engine/Source/Editor/DataTableEditor/Private/DataTableEditor.h
@@ -5,6 +5,7 @@
 #include "IDataTableEditor.h"
 #include "Toolkits/AssetEditorToolkit.h"
 
+class UDataTable;
 
 /** Viewer/editor for a DataTable */
 class FDataTableEditor :
diff --git a/Engine/Source/Editor/DataTableEditor/Private/DataTableEditorModule.cpp b/Engine/Source/Editor/DataTableEditor/Private/DataTableEditorModule.cpp
--- a/Engine/Source/Editor/DataTableEditor/Private/DataTableEditorModule.cpp
+++ b/Engine/Source/Editor/DataTableEditor/Private/DataTableEditorModule.cpp
@@ -4,6 +4,7 @@
 #include "ModuleManager.h"
 #include "DataTableEditor.h"
 #include "Toolkits/ToolkitManager.h"
+#include "Engine/DataTable.h"
 
 IMPLEMENT_MODULE( FDataTableEditorModule, DataTableEditor );
 
diff --git a/Engine/Source/Editor/DataTableEditor/Public/DataTableEditorModule.h b/Engine/Source/Editor/DataTableEditor/Public/DataTableEditorModule.h
--- a/Engine/Source/Editor/DataTableEditor/Public/DataTableEditorModule.h
+++ b/Engine/Source/Editor/DataTableEditor/Public/DataTableEditorModule.h
@@ -7,6 +7,7 @@
 #include "Toolkits/AssetEditorToolkit.h" // For FExtensibilityManager
 
 class IDataTableEditor;
+class UDataTable;
 
 /** DataTable Editor module */
 class FDataTableEditorModule : public IModuleInterface,
diff --git a/Engine/Source/Editor/DestructibleMeshEditor/Private/DestructibleMeshEditor.cpp b/Engine/Source/Editor/DestructibleMeshEditor/Private/DestructibleMeshEditor.cpp
--- a/Engine/Source/Editor/DestructibleMeshEditor/Private/DestructibleMeshEditor.cpp
+++ b/Engine/Source/Editor/DestructibleMeshEditor/Private/DestructibleMeshEditor.cpp
@@ -10,6 +10,7 @@
 #include "SDockTab.h"
 
 #include "Engine/DestructibleMesh.h"
+#include "Engine/Selection.h"
 
 #define LOCTEXT_NAMESPACE "DestructibleMeshEditor"
 
diff --git a/Engine/Source/Editor/DestructibleMeshEditor/Private/DestructibleMeshEditorModule.cpp b/Engine/Source/Editor/DestructibleMeshEditor/Private/DestructibleMeshEditorModule.cpp
--- a/Engine/Source/Editor/DestructibleMeshEditor/Private/DestructibleMeshEditorModule.cpp
+++ b/Engine/Source/Editor/DestructibleMeshEditor/Private/DestructibleMeshEditorModule.cpp
@@ -9,6 +9,7 @@
 
 #include "ApexDestructibleAssetImport.h"
 #include "Engine/DestructibleMesh.h"
+#include "Engine/StaticMesh.h"
 
 IMPLEMENT_MODULE( FDestructibleMeshEditorModule, DestructibleMeshEditor );
 
diff --git a/Engine/Source/Editor/DestructibleMeshEditor/Private/SDestructibleMeshEditorViewport.cpp b/Engine/Source/Editor/DestructibleMeshEditor/Private/SDestructibleMeshEditorViewport.cpp
--- a/Engine/Source/Editor/DestructibleMeshEditor/Private/SDestructibleMeshEditorViewport.cpp
+++ b/Engine/Source/Editor/DestructibleMeshEditor/Private/SDestructibleMeshEditorViewport.cpp
@@ -12,6 +12,8 @@
 #include "ComponentReregisterContext.h"
 #include "Engine/DestructibleMesh.h"
 #include "SDockableTab.h"
+#include "Components/DestructibleComponent.h"
+#include "Engine/StaticMesh.h"
 
 DEFINE_LOG_CATEGORY_STATIC(LogDestructibleMeshEditor, Log, All);
 
diff --git a/Engine/Source/Editor/DetailCustomizations/Private/ActorMaterialCategory.cpp b/Engine/Source/Editor/DetailCustomizations/Private/ActorMaterialCategory.cpp
--- a/Engine/Source/Editor/DetailCustomizations/Private/ActorMaterialCategory.cpp
+++ b/Engine/Source/Editor/DetailCustomizations/Private/ActorMaterialCategory.cpp
@@ -7,6 +7,9 @@
 
 #include "LandscapeProxy.h"
 #include "LandscapeComponent.h"
+#include "Components/DecalComponent.h"
+#include "Components/TextRenderComponent.h"
+#include "AI/Navigation/NavigationSystem.h"
 
 #define LOCTEXT_NAMESPACE "SMaterialList"
 
diff --git a/Engine/Source/Editor/DetailCustomizations/Private/AmbientSoundDetails.cpp b/Engine/Source/Editor/DetailCustomizations/Private/AmbientSoundDetails.cpp
--- a/Engine/Source/Editor/DetailCustomizations/Private/AmbientSoundDetails.cpp
+++ b/Engine/Source/Editor/DetailCustomizations/Private/AmbientSoundDetails.cpp
@@ -4,11 +4,14 @@
 #include "AmbientSoundDetails.h"
 #include "Toolkits/AssetEditorManager.h"
 #include "SoundDefinitions.h"
+#include "Sound/AmbientSound.h"
 #include "Sound/SoundNodeLooping.h"
 #include "Sound/SoundNodeMixer.h"
 #include "Sound/SoundNodeDelay.h"
 #include "Sound/SoundNodeRandom.h"
 #include "Sound/SoundNodeWavePlayer.h"
+#include "Sound/SoundCue.h"
+#include "Components/AudioComponent.h"
 
 
 #define LOCTEXT_NAMESPACE "AmbientSoundDetails"
diff --git a/Engine/Source/Editor/DetailCustomizations/Private/AnimTransitionNodeDetails.h b/Engine/Source/Editor/DetailCustomizations/Private/AnimTransitionNodeDetails.h
--- a/Engine/Source/Editor/DetailCustomizations/Private/AnimTransitionNodeDetails.h
+++ b/Engine/Source/Editor/DetailCustomizations/Private/AnimTransitionNodeDetails.h
@@ -4,6 +4,9 @@
 
 #include "PreviewScene.h"
 #include "AnimGraphDefinitions.h"
+#include "IDetailCustomization.h"
+
+class UAnimStateTransitionNode;
 
 class FAnimTransitionNodeDetails : public IDetailCustomization
 {
diff --git a/Engine/Source/Editor/DetailCustomizations/Private/AttenuationSettingsCustomizations.cpp b/Engine/Source/Editor/DetailCustomizations/Private/AttenuationSettingsCustomizations.cpp
--- a/Engine/Source/Editor/DetailCustomizations/Private/AttenuationSettingsCustomizations.cpp
+++ b/Engine/Source/Editor/DetailCustomizations/Private/AttenuationSettingsCustomizations.cpp
@@ -2,6 +2,7 @@
 
 #include "DetailCustomizationsPrivatePCH.h"
 #include "AttenuationSettingsCustomizations.h"
+#include "Sound/SoundAttenuation.h"
 
 TSharedRef<IPropertyTypeCustomization> FAttenuationSettingsCustomization::MakeInstance() 
 {
diff --git a/Engine/Source/Editor/DetailCustomizations/Private/BodyInstanceCustomization.cpp b/Engine/Source/Editor/DetailCustomizations/Private/BodyInstanceCustomization.cpp
--- a/Engine/Source/Editor/DetailCustomizations/Private/BodyInstanceCustomization.cpp
+++ b/Engine/Source/Editor/DetailCustomizations/Private/BodyInstanceCustomization.cpp
@@ -4,6 +4,7 @@
 #include "BodyInstanceCustomization.h"
 #include "ScopedTransaction.h"
 #include "Editor/Documentation/Public/IDocumentation.h"
+#include "Engine/CollisionProfile.h"
 
 #define LOCTEXT_NAMESPACE "BodyInstanceCustomization"
 
diff --git a/Engine/Source/Editor/DetailCustomizations/Private/BodySetupDetails.cpp b/Engine/Source/Editor/DetailCustomizations/Private/BodySetupDetails.cpp
--- a/Engine/Source/Editor/DetailCustomizations/Private/BodySetupDetails.cpp
+++ b/Engine/Source/Editor/DetailCustomizations/Private/BodySetupDetails.cpp
@@ -5,6 +5,7 @@
 #include "ScopedTransaction.h"
 #include "ObjectEditorUtils.h"
 #include "IDocumentation.h"
+#include "PhysicsEngine/BodySetup.h"
 
 #define LOCTEXT_NAMESPACE "BodySetupDetails"
 
diff --git a/Engine/Source/Editor/DetailCustomizations/Private/BrushDetails.cpp b/Engine/Source/Editor/DetailCustomizations/Private/BrushDetails.cpp
--- a/Engine/Source/Editor/DetailCustomizations/Private/BrushDetails.cpp
+++ b/Engine/Source/Editor/DetailCustomizations/Private/BrushDetails.cpp
@@ -11,6 +11,7 @@
 #include "ClassIconFinder.h"
 #include "ScopedTransaction.h"
 #include "Engine/StaticMeshActor.h"
+#include "GameFramework/Volume.h"
 
 #define LOCTEXT_NAMESPACE "BrushDetails"
 
diff --git a/Engine/Source/Editor/DetailCustomizations/Private/CollisionProfileDetails.cpp b/Engine/Source/Editor/DetailCustomizations/Private/CollisionProfileDetails.cpp
--- a/Engine/Source/Editor/DetailCustomizations/Private/CollisionProfileDetails.cpp
+++ b/Engine/Source/Editor/DetailCustomizations/Private/CollisionProfileDetails.cpp
@@ -6,6 +6,7 @@
 #include "ObjectEditorUtils.h"
 #include "BodyInstanceCustomization.h"
 #include "IDocumentation.h"
+#include "Engine/CollisionProfile.h"
 
 #define LOCTEXT_NAMESPACE "CollsiionProfileDetails"
 
diff --git a/Engine/Source/Editor/DetailCustomizations/Private/CollisionProfileDetails.h b/Engine/Source/Editor/DetailCustomizations/Private/CollisionProfileDetails.h
--- a/Engine/Source/Editor/DetailCustomizations/Private/CollisionProfileDetails.h
+++ b/Engine/Source/Editor/DetailCustomizations/Private/CollisionProfileDetails.h
@@ -1,6 +1,7 @@
 // Copyright 1998-2014 Epic Games, Inc. All Rights Reserved.
 
 #pragma once
+#include "Engine/CollisionProfile.h"
 
 // Class containing the friend information - used to build the list view
 class FChannelListItem
diff --git a/Engine/Source/Editor/DetailCustomizations/Private/CurveColorCustomization.cpp b/Engine/Source/Editor/DetailCustomizations/Private/CurveColorCustomization.cpp
--- a/Engine/Source/Editor/DetailCustomizations/Private/CurveColorCustomization.cpp
+++ b/Engine/Source/Editor/DetailCustomizations/Private/CurveColorCustomization.cpp
@@ -8,6 +8,7 @@
 #include "AssetRegistryModule.h"
 #include "SCurveEditor.h"
 #include "Curves/CurveLinearColor.h"
+#include "Engine/Selection.h"
 
 #define LOCTEXT_NAMESPACE "CurveColorCustomization"
 
diff --git a/Engine/Source/Editor/DetailCustomizations/Private/CurveStructCustomization.cpp b/Engine/Source/Editor/DetailCustomizations/Private/CurveStructCustomization.cpp
--- a/Engine/Source/Editor/DetailCustomizations/Private/CurveStructCustomization.cpp
+++ b/Engine/Source/Editor/DetailCustomizations/Private/CurveStructCustomization.cpp
@@ -7,6 +7,8 @@
 #include "MiniCurveEditor.h"
 #include "AssetRegistryModule.h"
 #include "SCurveEditor.h"
+#include "Curves/CurveFloat.h"
+#include "Engine/Selection.h"
 
 #define LOCTEXT_NAMESPACE "CurveStructCustomization"
 
diff --git a/Engine/Source/Editor/DetailCustomizations/Private/CurveTableCustomization.h b/Engine/Source/Editor/DetailCustomizations/Private/CurveTableCustomization.h
--- a/Engine/Source/Editor/DetailCustomizations/Private/CurveTableCustomization.h
+++ b/Engine/Source/Editor/DetailCustomizations/Private/CurveTableCustomization.h
@@ -1,6 +1,7 @@
 // Copyright 1998-2014 Epic Games, Inc. All Rights Reserved.
 
 #pragma once
+#include "Engine/CurveTable.h"
 
 /**
  * Customizes a DataTable asset to use a dropdown
diff --git a/Engine/Source/Editor/DetailCustomizations/Private/DataTableCustomization.h b/Engine/Source/Editor/DetailCustomizations/Private/DataTableCustomization.h
--- a/Engine/Source/Editor/DetailCustomizations/Private/DataTableCustomization.h
+++ b/Engine/Source/Editor/DetailCustomizations/Private/DataTableCustomization.h
@@ -1,6 +1,7 @@
 // Copyright 1998-2014 Epic Games, Inc. All Rights Reserved.
 
 #pragma once
+#include "Engine/DataTable.h"
 
 #define LOCTEXT_NAMESPACE "FDataTableCustomizationLayout"
 
diff --git a/Engine/Source/Editor/DetailCustomizations/Private/DialogueStructsCustomizations.cpp b/Engine/Source/Editor/DetailCustomizations/Private/DialogueStructsCustomizations.cpp
--- a/Engine/Source/Editor/DetailCustomizations/Private/DialogueStructsCustomizations.cpp
+++ b/Engine/Source/Editor/DetailCustomizations/Private/DialogueStructsCustomizations.cpp
@@ -6,6 +6,7 @@
 #include "DialogueWaveWidgets.h"
 #include "AssetThumbnail.h"
 #include "Sound/DialogueWave.h"
+#include "Sound/DialogueVoice.h"
 
 #define LOCTEXT_NAMESPACE "DialogueWaveDetails"
 
diff --git a/Engine/Source/Editor/DetailCustomizations/Private/DialogueWaveWidgets.cpp b/Engine/Source/Editor/DetailCustomizations/Private/DialogueWaveWidgets.cpp
--- a/Engine/Source/Editor/DetailCustomizations/Private/DialogueWaveWidgets.cpp
+++ b/Engine/Source/Editor/DetailCustomizations/Private/DialogueWaveWidgets.cpp
@@ -7,6 +7,7 @@
 #include "SAssetDropTarget.h"
 #include "AssetRegistryModule.h"
 #include "Sound/DialogueVoice.h"
+#include "Engine/Selection.h"
 
 #define LOCTEXT_NAMESPACE "DialogueWaveDetails"
 
diff --git a/Engine/Source/Editor/DetailCustomizations/Private/DirectionalLightComponentDetails.cpp b/Engine/Source/Editor/DetailCustomizations/Private/DirectionalLightComponentDetails.cpp
--- a/Engine/Source/Editor/DetailCustomizations/Private/DirectionalLightComponentDetails.cpp
+++ b/Engine/Source/Editor/DetailCustomizations/Private/DirectionalLightComponentDetails.cpp
@@ -2,6 +2,7 @@
 
 #include "DetailCustomizationsPrivatePCH.h"
 #include "DirectionalLightComponentDetails.h"
+#include "Components/LightComponentBase.h"
 
 #define LOCTEXT_NAMESPACE "DirectionalLightComponentDetails"
 
diff --git a/Engine/Source/Editor/DetailCustomizations/Private/FbxImportUIDetails.cpp b/Engine/Source/Editor/DetailCustomizations/Private/FbxImportUIDetails.cpp
--- a/Engine/Source/Editor/DetailCustomizations/Private/FbxImportUIDetails.cpp
+++ b/Engine/Source/Editor/DetailCustomizations/Private/FbxImportUIDetails.cpp
@@ -4,6 +4,7 @@
 #include "FbxImportUIDetails.h"
 #include "Factories/FbxAnimSequenceImportData.h"
 #include "STextComboBox.h"
+#include "Engine/StaticMesh.h"
 
 #define LOCTEXT_NAMESPACE "FbxImportUIDetails"
 
diff --git a/Engine/Source/Editor/DetailCustomizations/Private/InputSettingsDetails.cpp b/Engine/Source/Editor/DetailCustomizations/Private/InputSettingsDetails.cpp
--- a/Engine/Source/Editor/DetailCustomizations/Private/InputSettingsDetails.cpp
+++ b/Engine/Source/Editor/DetailCustomizations/Private/InputSettingsDetails.cpp
@@ -4,6 +4,7 @@
 #include "InputSettingsDetails.h"
 #include "ScopedTransaction.h"
 #include "IDocumentation.h"
+#include "GameFramework/InputSettings.h"
 
 #define LOCTEXT_NAMESPACE "InputSettingsDetails"
 
diff --git a/Engine/Source/Editor/DetailCustomizations/Private/InputStructCustomization.cpp b/Engine/Source/Editor/DetailCustomizations/Private/InputStructCustomization.cpp
--- a/Engine/Source/Editor/DetailCustomizations/Private/InputStructCustomization.cpp
+++ b/Engine/Source/Editor/DetailCustomizations/Private/InputStructCustomization.cpp
@@ -3,6 +3,7 @@
 #include "DetailCustomizationsPrivatePCH.h"
 #include "InputStructCustomization.h"
 #include "InputSettingsDetails.h"
+#include "GameFramework/PlayerInput.h"
 
 #define LOCTEXT_NAMESPACE "InputStructCustomization"
 
diff --git a/Engine/Source/Editor/DetailCustomizations/Private/LightComponentDetails.cpp b/Engine/Source/Editor/DetailCustomizations/Private/LightComponentDetails.cpp
--- a/Engine/Source/Editor/DetailCustomizations/Private/LightComponentDetails.cpp
+++ b/Engine/Source/Editor/DetailCustomizations/Private/LightComponentDetails.cpp
@@ -2,6 +2,7 @@
 
 #include "DetailCustomizationsPrivatePCH.h"
 #include "LightComponentDetails.h"
+#include "Components/LightComponentBase.h"
 
 #define LOCTEXT_NAMESPACE "LightComponentDetails"
 
diff --git a/Engine/Source/Editor/DetailCustomizations/Private/NavLinkStructCustomization.cpp b/Engine/Source/Editor/DetailCustomizations/Private/NavLinkStructCustomization.cpp
--- a/Engine/Source/Editor/DetailCustomizations/Private/NavLinkStructCustomization.cpp
+++ b/Engine/Source/Editor/DetailCustomizations/Private/NavLinkStructCustomization.cpp
@@ -2,6 +2,8 @@
 
 #include "DetailCustomizationsPrivatePCH.h"
 #include "NavLinkStructCustomization.h"
+#include "AI/Navigation/NavLinkDefinition.h"
+#include "AI/Navigation/NavigationSystem.h"
 
 #define LOCTEXT_NAMESPACE "FNavLinkStructCustomization"
 
diff --git a/Engine/Source/Editor/DetailCustomizations/Private/PhysicsSettingsDetails.cpp b/Engine/Source/Editor/DetailCustomizations/Private/PhysicsSettingsDetails.cpp
--- a/Engine/Source/Editor/DetailCustomizations/Private/PhysicsSettingsDetails.cpp
+++ b/Engine/Source/Editor/DetailCustomizations/Private/PhysicsSettingsDetails.cpp
@@ -6,6 +6,7 @@
 #include "ObjectEditorUtils.h"
 #include "BodyInstanceCustomization.h"
 #include "IDocumentation.h"
+#include "STableViewBase.h"
 
 #define LOCTEXT_NAMESPACE "PhysicalSurfaceDetails"
 
diff --git a/Engine/Source/Editor/DetailCustomizations/Private/PhysicsSettingsDetails.h b/Engine/Source/Editor/DetailCustomizations/Private/PhysicsSettingsDetails.h
--- a/Engine/Source/Editor/DetailCustomizations/Private/PhysicsSettingsDetails.h
+++ b/Engine/Source/Editor/DetailCustomizations/Private/PhysicsSettingsDetails.h
@@ -1,6 +1,7 @@
 // Copyright 1998-2014 Epic Games, Inc. All Rights Reserved.
 
 #pragma once
+#include "PhysicsEngine/PhysicsSettings.h"
 
 DECLARE_DELEGATE(FOnCommitChange)
 
diff --git a/Engine/Source/Editor/DetailCustomizations/Private/PointLightComponentDetails.cpp b/Engine/Source/Editor/DetailCustomizations/Private/PointLightComponentDetails.cpp
--- a/Engine/Source/Editor/DetailCustomizations/Private/PointLightComponentDetails.cpp
+++ b/Engine/Source/Editor/DetailCustomizations/Private/PointLightComponentDetails.cpp
@@ -2,6 +2,7 @@
 
 #include "DetailCustomizationsPrivatePCH.h"
 #include "PointLightComponentDetails.h"
+#include "Components/LightComponentBase.h"
 
 #define LOCTEXT_NAMESPACE "PointLightComponentDetails"
 
diff --git a/Engine/Source/Editor/DetailCustomizations/Private/PrimitiveComponentDetails.cpp b/Engine/Source/Editor/DetailCustomizations/Private/PrimitiveComponentDetails.cpp
--- a/Engine/Source/Editor/DetailCustomizations/Private/PrimitiveComponentDetails.cpp
+++ b/Engine/Source/Editor/DetailCustomizations/Private/PrimitiveComponentDetails.cpp
@@ -10,6 +10,7 @@
 #include "Components/InstancedStaticMeshComponent.h"
 #include "IDocumentation.h"
 #include "EditorCategoryUtils.h"
+#include "PhysicsEngine/PhysicsSettings.h"
 
 
 #define LOCTEXT_NAMESPACE "PrimitiveComponentDetails"
diff --git a/Engine/Source/Editor/DetailCustomizations/Private/ReflectionCaptureDetails.h b/Engine/Source/Editor/DetailCustomizations/Private/ReflectionCaptureDetails.h
--- a/Engine/Source/Editor/DetailCustomizations/Private/ReflectionCaptureDetails.h
+++ b/Engine/Source/Editor/DetailCustomizations/Private/ReflectionCaptureDetails.h
@@ -1,6 +1,8 @@
 // Copyright 1998-2014 Epic Games, Inc. All Rights Reserved.
 
 #pragma once
+#include "Engine/ReflectionCapture.h"
+#include "IDetailCustomization.h"
 
 class FReflectionCaptureDetails : public IDetailCustomization
 {
diff --git a/Engine/Source/Editor/DetailCustomizations/Private/SceneCaptureDetails.cpp b/Engine/Source/Editor/DetailCustomizations/Private/SceneCaptureDetails.cpp
--- a/Engine/Source/Editor/DetailCustomizations/Private/SceneCaptureDetails.cpp
+++ b/Engine/Source/Editor/DetailCustomizations/Private/SceneCaptureDetails.cpp
@@ -3,6 +3,11 @@
 #include "DetailCustomizationsPrivatePCH.h"
 #include "Runtime/Engine/Public/ShowFlags.h"
 #include "SceneCaptureDetails.h"
+#include "Engine/SceneCapture2D.h"
+#include "Engine/SceneCaptureCube.h"
+#include "Components/SceneCaptureComponent.h"
+#include "Components/SceneCaptureComponent2D.h"
+#include "Components/SceneCaptureComponentCube.h"
 
 #define LOCTEXT_NAMESPACE "SceneCaptureDetails"
 
diff --git a/Engine/Source/Editor/DetailCustomizations/Private/SceneCaptureDetails.h b/Engine/Source/Editor/DetailCustomizations/Private/SceneCaptureDetails.h
--- a/Engine/Source/Editor/DetailCustomizations/Private/SceneCaptureDetails.h
+++ b/Engine/Source/Editor/DetailCustomizations/Private/SceneCaptureDetails.h
@@ -1,6 +1,8 @@
 // Copyright 1998-2014 Epic Games, Inc. All Rights Reserved.
 
 #pragma once
+#include "Components/SceneCaptureComponent.h"
+#include "IDetailCustomization.h"
 
 class FSceneCaptureDetails : public IDetailCustomization
 {
diff --git a/Engine/Source/Editor/DetailCustomizations/Private/SceneComponentDetails.cpp b/Engine/Source/Editor/DetailCustomizations/Private/SceneComponentDetails.cpp
--- a/Engine/Source/Editor/DetailCustomizations/Private/SceneComponentDetails.cpp
+++ b/Engine/Source/Editor/DetailCustomizations/Private/SceneComponentDetails.cpp
@@ -9,6 +9,10 @@
 #include "ComponentTransformDetails.h"
 #include "SNotificationList.h"
 #include "NotificationManager.h"
+#include "Engine/BlueprintGeneratedClass.h"
+#include "Engine/SCS_Node.h"
+#include "Engine/SimpleConstructionScript.h"
+#include "Components/LightComponentBase.h"
 
 #define LOCTEXT_NAMESPACE "SceneComponentDetails"
 
diff --git a/Engine/Source/Editor/DetailCustomizations/Private/SkeletalMeshComponentDetails.cpp b/Engine/Source/Editor/DetailCustomizations/Private/SkeletalMeshComponentDetails.cpp
--- a/Engine/Source/Editor/DetailCustomizations/Private/SkeletalMeshComponentDetails.cpp
+++ b/Engine/Source/Editor/DetailCustomizations/Private/SkeletalMeshComponentDetails.cpp
@@ -6,6 +6,8 @@
 #include "AssetData.h"
 #include "Editor/ClassViewer/Public/ClassViewerModule.h"
 #include "Editor/ClassViewer/Public/ClassViewerFilter.h"
+#include "Engine/Selection.h"
+#include "Animation/AnimBlueprintGeneratedClass.h"
 
 #define LOCTEXT_NAMESPACE "SkeletalMeshComponentDetails"
 
diff --git a/Engine/Source/Editor/DetailCustomizations/Private/SkyLightComponentDetails.cpp b/Engine/Source/Editor/DetailCustomizations/Private/SkyLightComponentDetails.cpp
--- a/Engine/Source/Editor/DetailCustomizations/Private/SkyLightComponentDetails.cpp
+++ b/Engine/Source/Editor/DetailCustomizations/Private/SkyLightComponentDetails.cpp
@@ -2,6 +2,8 @@
 
 #include "DetailCustomizationsPrivatePCH.h"
 #include "SkyLightComponentDetails.h"
+#include "Components/LightComponentBase.h"
+#include "Engine/SkyLight.h"
 
 #define LOCTEXT_NAMESPACE "SkyLightComponentDetails"
 
diff --git a/Engine/Source/Editor/DetailCustomizations/Private/SkyLightComponentDetails.h b/Engine/Source/Editor/DetailCustomizations/Private/SkyLightComponentDetails.h
--- a/Engine/Source/Editor/DetailCustomizations/Private/SkyLightComponentDetails.h
+++ b/Engine/Source/Editor/DetailCustomizations/Private/SkyLightComponentDetails.h
@@ -1,6 +1,8 @@
 // Copyright 1998-2014 Epic Games, Inc. All Rights Reserved.
 
 #pragma once
+#include "Engine/SkyLight.h"
+#include "IDetailCustomization.h"
 
 class FSkyLightComponentDetails : public IDetailCustomization
 {
diff --git a/Engine/Source/Editor/DetailCustomizations/Private/SlateFontInfoCustomization.cpp b/Engine/Source/Editor/DetailCustomizations/Private/SlateFontInfoCustomization.cpp
--- a/Engine/Source/Editor/DetailCustomizations/Private/SlateFontInfoCustomization.cpp
+++ b/Engine/Source/Editor/DetailCustomizations/Private/SlateFontInfoCustomization.cpp
@@ -2,6 +2,8 @@
 
 #include "DetailCustomizationsPrivatePCH.h"
 #include "SlateFontInfoCustomization.h"
+#include "AssetData.h"
+#include "Engine/Font.h"
 
 #define LOCTEXT_NAMESPACE "SlateFontInfo"
 
diff --git a/Engine/Source/Editor/DetailCustomizations/Private/SlateSoundCustomization.cpp b/Engine/Source/Editor/DetailCustomizations/Private/SlateSoundCustomization.cpp
--- a/Engine/Source/Editor/DetailCustomizations/Private/SlateSoundCustomization.cpp
+++ b/Engine/Source/Editor/DetailCustomizations/Private/SlateSoundCustomization.cpp
@@ -2,6 +2,7 @@
 
 #include "DetailCustomizationsPrivatePCH.h"
 #include "SlateSoundCustomization.h"
+#include "Sound/SoundBase.h"
 
 TSharedRef<IPropertyTypeCustomization> FSlateSoundStructCustomization::MakeInstance() 
 {
diff --git a/Engine/Source/Editor/DetailCustomizations/Private/SoundWaveDetails.cpp b/Engine/Source/Editor/DetailCustomizations/Private/SoundWaveDetails.cpp
--- a/Engine/Source/Editor/DetailCustomizations/Private/SoundWaveDetails.cpp
+++ b/Engine/Source/Editor/DetailCustomizations/Private/SoundWaveDetails.cpp
@@ -2,6 +2,7 @@
 
 #include "DetailCustomizationsPrivatePCH.h"
 #include "SoundWaveDetails.h"
+#include "Sound/SoundWave.h"
 
 TSharedRef<IDetailCustomization> FSoundWaveDetails::MakeInstance()
 {
diff --git a/Engine/Source/Editor/DetailCustomizations/Private/VehicleTransmissionDataCustomization.cpp b/Engine/Source/Editor/DetailCustomizations/Private/VehicleTransmissionDataCustomization.cpp
--- a/Engine/Source/Editor/DetailCustomizations/Private/VehicleTransmissionDataCustomization.cpp
+++ b/Engine/Source/Editor/DetailCustomizations/Private/VehicleTransmissionDataCustomization.cpp
@@ -4,6 +4,7 @@
 #include "VehicleTransmissionDataCustomization.h"
 #include "ScopedTransaction.h"
 #include "Editor/Documentation/Public/IDocumentation.h"
+#include "Vehicles/WheeledVehicleMovementComponent4W.h"
 
 #define LOCTEXT_NAMESPACE "VehicleTransmissionDataCustomization"
 
diff --git a/Engine/Source/Editor/DetailCustomizations/Private/VehicleTransmissionDataCustomization.h b/Engine/Source/Editor/DetailCustomizations/Private/VehicleTransmissionDataCustomization.h
--- a/Engine/Source/Editor/DetailCustomizations/Private/VehicleTransmissionDataCustomization.h
+++ b/Engine/Source/Editor/DetailCustomizations/Private/VehicleTransmissionDataCustomization.h
@@ -1,6 +1,9 @@
 // Copyright 1998-2014 Epic Games, Inc. All Rights Reserved.
 
 #pragma once
+#include "IPropertyTypeCustomization.h"
+
+struct FVehicleTransmissionData;
 
 class FVehicleTransmissionDataCustomization : public IPropertyTypeCustomization
 {
diff --git a/Engine/Source/Editor/DetailCustomizations/Private/WheeledVehicleMovementComponent4WDetails.cpp b/Engine/Source/Editor/DetailCustomizations/Private/WheeledVehicleMovementComponent4WDetails.cpp
--- a/Engine/Source/Editor/DetailCustomizations/Private/WheeledVehicleMovementComponent4WDetails.cpp
+++ b/Engine/Source/Editor/DetailCustomizations/Private/WheeledVehicleMovementComponent4WDetails.cpp
@@ -5,6 +5,7 @@
 #include "ScopedTransaction.h"
 #include "ObjectEditorUtils.h"
 #include "Vehicles/WheeledVehicleMovementComponent4W.h"
+#include "WheeledVehicleMovementComponent4WDetails.h"
 #include "IDocumentation.h"
 #include "SCurveEditor.h"
 
diff --git a/Engine/Source/Editor/DetailCustomizations/Private/WheeledVehicleMovementComponent4WDetails.h b/Engine/Source/Editor/DetailCustomizations/Private/WheeledVehicleMovementComponent4WDetails.h
--- a/Engine/Source/Editor/DetailCustomizations/Private/WheeledVehicleMovementComponent4WDetails.h
+++ b/Engine/Source/Editor/DetailCustomizations/Private/WheeledVehicleMovementComponent4WDetails.h
@@ -3,6 +3,8 @@
 #pragma once
 
 #include "SCurveEditor.h"
+#include "IDetailCustomization.h"
+class UWheeledVehicleMovementComponent4W;
 
 class FWheeledVehicleMovementComponent4WDetails : public IDetailCustomization
 {
diff --git a/Engine/Source/Editor/DistCurveEditor/Private/CurveEditorViewportClient.cpp b/Engine/Source/Editor/DistCurveEditor/Private/CurveEditorViewportClient.cpp
--- a/Engine/Source/Editor/DistCurveEditor/Private/CurveEditorViewportClient.cpp
+++ b/Engine/Source/Editor/DistCurveEditor/Private/CurveEditorViewportClient.cpp
@@ -9,6 +9,7 @@
 #include "CurveEditorViewportClient.h"
 #include "Runtime/Engine/Public/Slate/SceneViewport.h"
 #include "Engine/InterpCurveEdSetup.h"
+#include "CanvasTypes.h"
 
 FCurveEditorViewportClient::FCurveEditorViewportClient(TWeakPtr<SDistributionCurveEditor> InCurveEditor, TWeakPtr<SCurveEditorViewport> InCurveEditorViewport)
 	: CurveEditorPtr(InCurveEditor)
diff --git a/Engine/Source/Editor/DistCurveEditor/Private/SDistributionCurveEditor.cpp b/Engine/Source/Editor/DistCurveEditor/Private/SDistributionCurveEditor.cpp
--- a/Engine/Source/Editor/DistCurveEditor/Private/SDistributionCurveEditor.cpp
+++ b/Engine/Source/Editor/DistCurveEditor/Private/SDistributionCurveEditor.cpp
@@ -12,6 +12,7 @@
 #include "STextEntryPopup.h"
 #include "SNotificationList.h"
 #include "NotificationManager.h"
+#include "Engine/InterpCurveEdSetup.h"
 
 #define LOCTEXT_NAMESPACE "CurveEditor"
 
diff --git a/Engine/Source/Editor/EnvironmentQueryEditor/Private/EnvironmentQueryEditor.h b/Engine/Source/Editor/EnvironmentQueryEditor/Private/EnvironmentQueryEditor.h
--- a/Engine/Source/Editor/EnvironmentQueryEditor/Private/EnvironmentQueryEditor.h
+++ b/Engine/Source/Editor/EnvironmentQueryEditor/Private/EnvironmentQueryEditor.h
@@ -4,6 +4,7 @@
 
 #include "IEnvironmentQueryEditor.h"
 #include "Toolkits/AssetEditorToolkit.h"
+#include "EditorUndoClient.h"
 
 
 /** Viewer/editor for a DataTable */
diff --git a/Engine/Source/Editor/FoliageEdit/Private/FoliageEdMode.cpp b/Engine/Source/Editor/FoliageEdit/Private/FoliageEdMode.cpp
--- a/Engine/Source/Editor/FoliageEdit/Private/FoliageEdMode.cpp
+++ b/Engine/Source/Editor/FoliageEdit/Private/FoliageEdMode.cpp
@@ -24,6 +24,11 @@
 #include "LandscapeHeightfieldCollisionComponent.h"
 #include "LandscapeInfo.h"
 #include "Components/SplineMeshComponent.h"
+#include "Engine/StaticMeshActor.h"
+#include "GameFramework/WorldSettings.h"
+#include "Engine/CollisionProfile.h"
+#include "Components/ModelComponent.h"
+#include "EngineUtils.h"
 
 #define FOLIAGE_SNAP_TRACE (10000.f)
 
diff --git a/Engine/Source/Editor/FoliageEdit/Private/SFoliageEdit.cpp b/Engine/Source/Editor/FoliageEdit/Private/SFoliageEdit.cpp
--- a/Engine/Source/Editor/FoliageEdit/Private/SFoliageEdit.cpp
+++ b/Engine/Source/Editor/FoliageEdit/Private/SFoliageEdit.cpp
@@ -14,6 +14,7 @@
 #include "Editor/IntroTutorials/Public/IIntroTutorials.h"
 
 #include "SFoliageEditMeshDisplayItem.h"
+#include "Engine/StaticMesh.h"
 
 #define LOCTEXT_NAMESPACE "FoliageEd_Mode"
 
diff --git a/Engine/Source/Editor/FoliageEdit/Private/SFoliageEditMeshDisplayItem.cpp b/Engine/Source/Editor/FoliageEdit/Private/SFoliageEditMeshDisplayItem.cpp
--- a/Engine/Source/Editor/FoliageEdit/Private/SFoliageEditMeshDisplayItem.cpp
+++ b/Engine/Source/Editor/FoliageEdit/Private/SFoliageEditMeshDisplayItem.cpp
@@ -15,6 +15,8 @@
 #include "SFoliageEditMeshDisplayItem.h"
 #include "ObjectEditorUtils.h"
 #include "TutorialMetaData.h"
+#include "Engine/Selection.h"
+#include "Engine/StaticMesh.h"
 
 #define LOCTEXT_NAMESPACE "FoliageEd_Mode"
 
diff --git a/Engine/Source/Editor/FontEditor/Private/FontEditor.cpp b/Engine/Source/Editor/FontEditor/Private/FontEditor.cpp
--- a/Engine/Source/Editor/FontEditor/Private/FontEditor.cpp
+++ b/Engine/Source/Editor/FontEditor/Private/FontEditor.cpp
@@ -13,6 +13,8 @@
 #include "Editor/PropertyEditor/Public/PropertyEditorModule.h"
 #include "Editor/PropertyEditor/Public/IDetailsView.h"
 #include "SDockTab.h"
+#include "Engine/Font.h"
+#include "Engine/Selection.h"
 
 #define LOCTEXT_NAMESPACE "FontEditor"
 
diff --git a/Engine/Source/Editor/FontEditor/Private/FontEditor.h b/Engine/Source/Editor/FontEditor/Private/FontEditor.h
--- a/Engine/Source/Editor/FontEditor/Private/FontEditor.h
+++ b/Engine/Source/Editor/FontEditor/Private/FontEditor.h
@@ -3,6 +3,8 @@
 #pragma once
 
 #include "Toolkits/AssetEditorToolkit.h"
+#include "EditorUndoClient.h"
+#include "Engine/Font.h"
 
 
 /*-----------------------------------------------------------------------------
diff --git a/Engine/Source/Editor/FontEditor/Private/SCompositeFontEditor.cpp b/Engine/Source/Editor/FontEditor/Private/SCompositeFontEditor.cpp
--- a/Engine/Source/Editor/FontEditor/Private/SCompositeFontEditor.cpp
+++ b/Engine/Source/Editor/FontEditor/Private/SCompositeFontEditor.cpp
@@ -7,6 +7,7 @@
 #include "SNumericEntryBox.h"
 #include "ScopedTransaction.h"
 #include "DesktopPlatformModule.h"
+#include "Engine/Font.h"
 
 #define LOCTEXT_NAMESPACE "FontEditor"
 
diff --git a/Engine/Source/Editor/FontEditor/Private/SFontEditorViewport.cpp b/Engine/Source/Editor/FontEditor/Private/SFontEditorViewport.cpp
--- a/Engine/Source/Editor/FontEditor/Private/SFontEditorViewport.cpp
+++ b/Engine/Source/Editor/FontEditor/Private/SFontEditorViewport.cpp
@@ -4,6 +4,10 @@
 #include "MouseDeltaTracker.h"
 #include "SFontEditorViewport.h"
 #include "Runtime/Engine/Public/Slate/SceneViewport.h"
+#include "CanvasTypes.h"
+#include "Engine/Font.h"
+#include "Engine/Selection.h"
+#include "CanvasItem.h"
 
 #define LOCTEXT_NAMESPACE "FontEditor"
 
diff --git a/Engine/Source/Editor/GameProjectGeneration/Private/GameProjectUtils.cpp b/Engine/Source/Editor/GameProjectGeneration/Private/GameProjectUtils.cpp
--- a/Engine/Source/Editor/GameProjectGeneration/Private/GameProjectUtils.cpp
+++ b/Engine/Source/Editor/GameProjectGeneration/Private/GameProjectUtils.cpp
@@ -20,6 +20,7 @@
 #include "DesktopPlatformModule.h"
 #include "SNotificationList.h"
 #include "NotificationManager.h"
+#include "GameFramework/GameMode.h"
 
 #define LOCTEXT_NAMESPACE "GameProjectUtils"
 
diff --git a/Engine/Source/Editor/GameProjectGeneration/Private/SNewClassDialog.cpp b/Engine/Source/Editor/GameProjectGeneration/Private/SNewClassDialog.cpp
--- a/Engine/Source/Editor/GameProjectGeneration/Private/SNewClassDialog.cpp
+++ b/Engine/Source/Editor/GameProjectGeneration/Private/SNewClassDialog.cpp
@@ -12,6 +12,17 @@
 #include "SWizard.h"
 #include "SHyperlink.h"
 #include "TutorialMetaData.h"
+#include "GameFramework/GameMode.h"
+#include "GameFramework/GameState.h"
+#include "GameFramework/Pawn.h"
+#include "GameFramework/WorldSettings.h"
+#include "GameFramework/PlayerController.h"
+#include "GameFramework/Character.h"
+#include "GameFramework/Actor.h"
+#include "Camera/PlayerCameraManager.h"
+#include "GameFramework/HUD.h"
+#include "GameFramework/PlayerState.h"
+#include "Kismet/BlueprintFunctionLibrary.h"
 
 
 #define LOCTEXT_NAMESPACE "GameProjectGeneration"
diff --git a/Engine/Source/Editor/GameProjectGeneration/Private/Tests/GameProjectAutomationTests.cpp b/Engine/Source/Editor/GameProjectGeneration/Private/Tests/GameProjectAutomationTests.cpp
--- a/Engine/Source/Editor/GameProjectGeneration/Private/Tests/GameProjectAutomationTests.cpp
+++ b/Engine/Source/Editor/GameProjectGeneration/Private/Tests/GameProjectAutomationTests.cpp
@@ -7,6 +7,7 @@
 #include "../GameProjectUtils.h"
 #include "Tests/AutomationTestSettings.h"
 #include "AutomationEditorCommon.h"
+#include "GameFramework/PlayerStart.h"
 
 DEFINE_LOG_CATEGORY_STATIC(LogGameProjectGenerationTests, Log, All);
 
diff --git a/Engine/Source/Editor/GameplayTagsEditor/Private/GameplayTagContainerCustomization.h b/Engine/Source/Editor/GameplayTagsEditor/Private/GameplayTagContainerCustomization.h
--- a/Engine/Source/Editor/GameplayTagsEditor/Private/GameplayTagContainerCustomization.h
+++ b/Engine/Source/Editor/GameplayTagsEditor/Private/GameplayTagContainerCustomization.h
@@ -4,6 +4,7 @@
 
 #include "Editor/PropertyEditor/Public/PropertyEditorModule.h"
 #include "SGameplayTagWidget.h"
+#include "EditorUndoClient.h"
 
 /** Customization for the gameplay tag container struct */
 class FGameplayTagContainerCustomization : public IPropertyTypeCustomization, public FEditorUndoClient
diff --git a/Engine/Source/Editor/GameplayTagsEditor/Private/GameplayTagCustomization.h b/Engine/Source/Editor/GameplayTagsEditor/Private/GameplayTagCustomization.h
--- a/Engine/Source/Editor/GameplayTagsEditor/Private/GameplayTagCustomization.h
+++ b/Engine/Source/Editor/GameplayTagsEditor/Private/GameplayTagCustomization.h
@@ -4,6 +4,7 @@
 
 #include "Editor/PropertyEditor/Public/PropertyEditorModule.h"
 #include "SGameplayTagWidget.h"
+#include "EditorUndoClient.h"
 
 /** Customization for the gameplay tag struct */
 class FGameplayTagCustomization : public IPropertyTypeCustomization, public FEditorUndoClient
diff --git a/Engine/Source/Editor/GeometryMode/Private/EditorGeometry.cpp b/Engine/Source/Editor/GeometryMode/Private/EditorGeometry.cpp
--- a/Engine/Source/Editor/GeometryMode/Private/EditorGeometry.cpp
+++ b/Engine/Source/Editor/GeometryMode/Private/EditorGeometry.cpp
@@ -3,6 +3,7 @@
 
 #include "GeometryModePrivatePCH.h"
 #include "EditorGeometry.h"
+#include "Engine/Polys.h"
 
 DEFINE_LOG_CATEGORY_STATIC(LogEditorGeometry, Log, All);
 
diff --git a/Engine/Source/Editor/GeometryMode/Private/GeometryEdMode.cpp b/Engine/Source/Editor/GeometryMode/Private/GeometryEdMode.cpp
--- a/Engine/Source/Editor/GeometryMode/Private/GeometryEdMode.cpp
+++ b/Engine/Source/Editor/GeometryMode/Private/GeometryEdMode.cpp
@@ -2,7 +2,7 @@
 
 #include "GeometryModePrivatePCH.h"
 #include "DynamicMeshBuilder.h"
-
+#include "Engine/Selection.h"
 
 IMPLEMENT_MODULE( FGeometryModeModule, GeometryMode );
 
diff --git a/Engine/Source/Editor/GeometryMode/Private/GeometryModifiers.cpp b/Engine/Source/Editor/GeometryMode/Private/GeometryModifiers.cpp
--- a/Engine/Source/Editor/GeometryMode/Private/GeometryModifiers.cpp
+++ b/Engine/Source/Editor/GeometryMode/Private/GeometryModifiers.cpp
@@ -10,6 +10,9 @@
 #include "ActorEditorUtils.h"
 #include "SNotificationList.h"
 #include "NotificationManager.h"
+#include "Engine/Polys.h"
+#include "Engine/Selection.h"
+#include "Components/BrushComponent.h"
 
 DEFINE_LOG_CATEGORY_STATIC(LogGeomModifier, Log, All);
 
diff --git a/Engine/Source/Editor/GeometryMode/Public/GeometryEdMode.h b/Engine/Source/Editor/GeometryMode/Public/GeometryEdMode.h
--- a/Engine/Source/Editor/GeometryMode/Public/GeometryEdMode.h
+++ b/Engine/Source/Editor/GeometryMode/Public/GeometryEdMode.h
@@ -1,7 +1,15 @@
 // Copyright 1998-2014 Epic Games, Inc. All Rights Reserved.
 
 #pragma once
-
+#include "EditorModes.h"
+#include "EditorModeTools.h"
+
+typedef FName FEditorModeID;
+class FGeomPoly;
+class FGeomEdge;
+class FGeomVertex;
+class FGeomBase;
+class FGeomObject;
 
 /**
  * Geometry mode module
diff --git a/Engine/Source/Editor/IntroTutorials/Private/EditorTutorialFactory.cpp b/Engine/Source/Editor/IntroTutorials/Private/EditorTutorialFactory.cpp
--- a/Engine/Source/Editor/IntroTutorials/Private/EditorTutorialFactory.cpp
+++ b/Engine/Source/Editor/IntroTutorials/Private/EditorTutorialFactory.cpp
@@ -5,6 +5,7 @@
 #include "EditorTutorial.h"
 #include "KismetEditorUtilities.h"
 #include "IAssetTypeActions.h"
+#include "Engine/BlueprintGeneratedClass.h"
 
 #define LOCTEXT_NAMESPACE "UEditorTutorialFactory"
 
diff --git a/Engine/Source/Editor/IntroTutorials/Private/EditorTutorialImportFactory.cpp b/Engine/Source/Editor/IntroTutorials/Private/EditorTutorialImportFactory.cpp
--- a/Engine/Source/Editor/IntroTutorials/Private/EditorTutorialImportFactory.cpp
+++ b/Engine/Source/Editor/IntroTutorials/Private/EditorTutorialImportFactory.cpp
@@ -7,6 +7,7 @@
 #include "IAssetTypeActions.h"
 #include "IDocumentation.h"
 #include "IDocumentationPage.h"
+#include "Engine/BlueprintGeneratedClass.h"
 
 #define LOCTEXT_NAMESPACE "UEditorTutorialImportFactory"
 
diff --git a/Engine/Source/Editor/IntroTutorials/Private/IntroTutorials.cpp b/Engine/Source/Editor/IntroTutorials/Private/IntroTutorials.cpp
--- a/Engine/Source/Editor/IntroTutorials/Private/IntroTutorials.cpp
+++ b/Engine/Source/Editor/IntroTutorials/Private/IntroTutorials.cpp
@@ -33,6 +33,7 @@
 #include "IDocumentation.h"
 #include "SDockTab.h"
 #include "ModuleManager.h"
+#include "IntroTutorials.h"
 
 
 #define LOCTEXT_NAMESPACE "IntroTutorials"
diff --git a/Engine/Source/Editor/IntroTutorials/Private/IntroTutorials.h b/Engine/Source/Editor/IntroTutorials/Private/IntroTutorials.h
--- a/Engine/Source/Editor/IntroTutorials/Private/IntroTutorials.h
+++ b/Engine/Source/Editor/IntroTutorials/Private/IntroTutorials.h
@@ -1,6 +1,7 @@
 // Copyright 1998-2014 Epic Games, Inc. All Rights Reserved.
 
 #pragma once
+#include "Curves/CurveFloat.h"
 
 /** Query delegate to see if we are in picking mode */
 DECLARE_DELEGATE_RetVal_OneParam(bool, FOnIsPicking, FName& /* OutWidgetNameToHighlight */);
diff --git a/Engine/Source/Editor/Kismet/Private/BlueprintActionMenuUtils.cpp b/Engine/Source/Editor/Kismet/Private/BlueprintActionMenuUtils.cpp
--- a/Engine/Source/Editor/Kismet/Private/BlueprintActionMenuUtils.cpp
+++ b/Engine/Source/Editor/Kismet/Private/BlueprintActionMenuUtils.cpp
@@ -18,6 +18,8 @@
 #include "KismetEditorUtilities.h"	// for CanPasteNodes()
 #include "K2ActionMenuBuilder.h"
 #include "BlueprintEditorSettings.h"
+#include "Engine/Selection.h"
+#include "Engine/LevelScriptActor.h"
 
 #define LOCTEXT_NAMESPACE "BlueprintActionMenuUtils"
 
diff --git a/Engine/Source/Editor/Kismet/Private/BlueprintDetailsCustomization.cpp b/Engine/Source/Editor/Kismet/Private/BlueprintDetailsCustomization.cpp
--- a/Engine/Source/Editor/Kismet/Private/BlueprintDetailsCustomization.cpp
+++ b/Engine/Source/Editor/Kismet/Private/BlueprintDetailsCustomization.cpp
@@ -34,6 +34,11 @@
 
 #include "IDocumentation.h"
 #include "STextComboBox.h"
+#include "Engine/UserDefinedStruct.h"
+#include "Engine/SimpleConstructionScript.h"
+#include "Engine/SCS_Node.h"
+#include "Components/TimelineComponent.h"
+#include "Engine/BlueprintGeneratedClass.h"
 
 #define LOCTEXT_NAMESPACE "BlueprintDetailsCustomization"
 
diff --git a/Engine/Source/Editor/Kismet/Private/BlueprintEditorCommands.cpp b/Engine/Source/Editor/Kismet/Private/BlueprintEditorCommands.cpp
--- a/Engine/Source/Editor/Kismet/Private/BlueprintEditorCommands.cpp
+++ b/Engine/Source/Editor/Kismet/Private/BlueprintEditorCommands.cpp
@@ -6,6 +6,8 @@
 #include "BlueprintEditorCommands.h"
 #include "Editor/UnrealEd/Public/Kismet2/BlueprintEditorUtils.h"
 #include "Editor/BlueprintGraph/Public/K2ActionMenuBuilder.h" // for AddEventForFunction(), ...
+#include "Engine/Selection.h"
+#include "Engine/LevelScriptActor.h"
 
 struct FBlueprintPaletteListBuilder;
 
diff --git a/Engine/Source/Editor/Kismet/Private/BlueprintEditorTabFactories.h b/Engine/Source/Editor/Kismet/Private/BlueprintEditorTabFactories.h
--- a/Engine/Source/Editor/Kismet/Private/BlueprintEditorTabFactories.h
+++ b/Engine/Source/Editor/Kismet/Private/BlueprintEditorTabFactories.h
@@ -1,6 +1,7 @@
 // Copyright 1998-2014 Epic Games, Inc. All Rights Reserved.
 
 #pragma once
+#include "Engine/TimelineTemplate.h"
 
 #define LOCTEXT_NAMESPACE "BlueprintEditor"
 
diff --git a/Engine/Source/Editor/Kismet/Private/BlueprintEditorViewportContextMenuExtender.cpp b/Engine/Source/Editor/Kismet/Private/BlueprintEditorViewportContextMenuExtender.cpp
--- a/Engine/Source/Editor/Kismet/Private/BlueprintEditorViewportContextMenuExtender.cpp
+++ b/Engine/Source/Editor/Kismet/Private/BlueprintEditorViewportContextMenuExtender.cpp
@@ -5,6 +5,7 @@
 #include "BlueprintEditorModule.h"
 #include "Editor/LevelEditor/Public/LevelEditor.h"
 #include "CreateBlueprintFromActorDialog.h"
+#include "Engine/Selection.h"
 
 DEFINE_LOG_CATEGORY_STATIC(LogViewportBlueprintMenu, Log, All);
 
diff --git a/Engine/Source/Editor/Kismet/Private/InstancedStaticMeshSCSEditorCustomization.cpp b/Engine/Source/Editor/Kismet/Private/InstancedStaticMeshSCSEditorCustomization.cpp
--- a/Engine/Source/Editor/Kismet/Private/InstancedStaticMeshSCSEditorCustomization.cpp
+++ b/Engine/Source/Editor/Kismet/Private/InstancedStaticMeshSCSEditorCustomization.cpp
@@ -4,6 +4,7 @@
 #include "InstancedStaticMeshSCSEditorCustomization.h"
 #include "BlueprintEditorModule.h"
 #include "InstancedFoliage.h"
+#include "Components/InstancedStaticMeshComponent.h"
 
 TSharedRef<ISCSEditorCustomization> FInstancedStaticMeshSCSEditorCustomization::MakeInstance(TSharedRef< IBlueprintEditor > InBlueprintEditor)
 {
diff --git a/Engine/Source/Editor/Kismet/Private/NativeCodeGenerationTool.cpp b/Engine/Source/Editor/Kismet/Private/NativeCodeGenerationTool.cpp
--- a/Engine/Source/Editor/Kismet/Private/NativeCodeGenerationTool.cpp
+++ b/Engine/Source/Editor/Kismet/Private/NativeCodeGenerationTool.cpp
@@ -3,6 +3,9 @@
 #include "BlueprintEditorPrivatePCH.h"
 #include "NativeCodeGenerationTool.h"
 #include "Dialogs.h"
+#include "Engine/BlueprintGeneratedClass.h"
+#include "Engine/UserDefinedEnum.h"
+#include "Engine/UserDefinedStruct.h"
 
 #define LOCTEXT_NAMESPACE "NativeCodeGenerationTool"
 
diff --git a/Engine/Source/Editor/Kismet/Private/SCSEditorViewportClient.cpp b/Engine/Source/Editor/Kismet/Private/SCSEditorViewportClient.cpp
--- a/Engine/Source/Editor/Kismet/Private/SCSEditorViewportClient.cpp
+++ b/Engine/Source/Editor/Kismet/Private/SCSEditorViewportClient.cpp
@@ -17,6 +17,13 @@
 #include "ComponentVisualizer.h"
 #include "InstancedFoliage.h"
 #include "Engine/SimpleConstructionScript.h"
+#include "CanvasTypes.h"
+#include "EngineUtils.h"
+#include "GameFramework/Actor.h"
+#include "Engine/SCS_Node.h"
+#include "Engine/TextureCube.h"
+#include "Components/InstancedStaticMeshComponent.h"
+#include "Engine/StaticMesh.h"
 
 DEFINE_LOG_CATEGORY_STATIC(LogSCSEditorViewport, Log, All);
 
diff --git a/Engine/Source/Editor/Kismet/Private/SKismetInspector.cpp b/Engine/Source/Editor/Kismet/Private/SKismetInspector.cpp
--- a/Engine/Source/Editor/Kismet/Private/SKismetInspector.cpp
+++ b/Engine/Source/Editor/Kismet/Private/SKismetInspector.cpp
@@ -18,6 +18,8 @@
 #include "UserDefinedEnumEditor.h"
 #include "UserDefinedStructureEditor.h"
 #include "FormatTextDetails.h"
+#include "Engine/SCS_Node.h"
+#include "Components/ChildActorComponent.h"
 
 #define LOCTEXT_NAMESPACE "KismetInspector"
 
diff --git a/Engine/Source/Editor/Kismet/Private/SMyBlueprint.cpp b/Engine/Source/Editor/Kismet/Private/SMyBlueprint.cpp
--- a/Engine/Source/Editor/Kismet/Private/SMyBlueprint.cpp
+++ b/Engine/Source/Editor/Kismet/Private/SMyBlueprint.cpp
@@ -39,6 +39,7 @@
 #include "SSearchBox.h"
 #include "GenericCommands.h"
 #include "Components/TimelineComponent.h"
+#include "Engine/TimelineTemplate.h"
 
 #define LOCTEXT_NAMESPACE "MyBlueprint"
 
diff --git a/Engine/Source/Editor/Kismet/Private/SSCSEditor.cpp b/Engine/Source/Editor/Kismet/Private/SSCSEditor.cpp
--- a/Engine/Source/Editor/Kismet/Private/SSCSEditor.cpp
+++ b/Engine/Source/Editor/Kismet/Private/SSCSEditor.cpp
@@ -29,6 +29,9 @@
 #include "TutorialMetaData.h"
 #include "SInlineEditableTextBlock.h"
 #include "GenericCommands.h"
+#include "Engine/SCS_Node.h"
+#include "Engine/SimpleConstructionScript.h"
+#include "Engine/Selection.h"
 
 #define LOCTEXT_NAMESPACE "SSCSEditor"
 
diff --git a/Engine/Source/Editor/Kismet/Private/STimelineEditor.cpp b/Engine/Source/Editor/Kismet/Private/STimelineEditor.cpp
--- a/Engine/Source/Editor/Kismet/Private/STimelineEditor.cpp
+++ b/Engine/Source/Editor/Kismet/Private/STimelineEditor.cpp
@@ -18,6 +18,10 @@
 #include "SNotificationList.h"
 #include "NotificationManager.h"
 #include "Engine/TimelineTemplate.h"
+#include "Curves/CurveFloat.h"
+#include "Curves/CurveLinearColor.h"
+#include "Curves/CurveVector.h"
+#include "Engine/Selection.h"
 
 #define LOCTEXT_NAMESPACE "STimelineEditor"
 
diff --git a/Engine/Source/Editor/Kismet/Private/UserDefinedEnumEditor.cpp b/Engine/Source/Editor/Kismet/Private/UserDefinedEnumEditor.cpp
--- a/Engine/Source/Editor/Kismet/Private/UserDefinedEnumEditor.cpp
+++ b/Engine/Source/Editor/Kismet/Private/UserDefinedEnumEditor.cpp
@@ -6,6 +6,7 @@
 
 #include "PropertyCustomizationHelpers.h"
 #include "SDockTab.h"
+#include "Engine/UserDefinedEnum.h"
 
 #define LOCTEXT_NAMESPACE "UserDefinedEnumEditor"
 
diff --git a/Engine/Source/Editor/Kismet/Private/UserDefinedStructureEditor.cpp b/Engine/Source/Editor/Kismet/Private/UserDefinedStructureEditor.cpp
--- a/Engine/Source/Editor/Kismet/Private/UserDefinedStructureEditor.cpp
+++ b/Engine/Source/Editor/Kismet/Private/UserDefinedStructureEditor.cpp
@@ -11,6 +11,7 @@
 #include "Editor/WorkspaceMenuStructure/Public/WorkspaceMenuStructureModule.h"
 #include "Editor/UnrealEd/Public/Kismet2/StructureEditorUtils.h"
 #include "SDockTab.h"
+#include "Engine/UserDefinedStruct.h"
 
 #define LOCTEXT_NAMESPACE "StructureEditor"
 
diff --git a/Engine/Source/Editor/Kismet/Public/BlueprintEditor.h b/Engine/Source/Editor/Kismet/Public/BlueprintEditor.h
--- a/Engine/Source/Editor/Kismet/Public/BlueprintEditor.h
+++ b/Engine/Source/Editor/Kismet/Public/BlueprintEditor.h
@@ -8,6 +8,10 @@
 #include "Toolkits/AssetEditorManager.h"
 #include "Kismet2/KismetEditorUtilities.h"
 #include "GraphEditor.h"
+#include "EditorUndoClient.h"
+#include "Engine/UserDefinedEnum.h"
+
+class USCS_Node;
 
 /////////////////////////////////////////////////////
 // FCustomDebugObjectEntry - Used to pass a custom debug object override around
diff --git a/Engine/Source/Editor/Kismet/Public/BlueprintEditorModule.h b/Engine/Source/Editor/Kismet/Public/BlueprintEditorModule.h
--- a/Engine/Source/Editor/Kismet/Public/BlueprintEditorModule.h
+++ b/Engine/Source/Editor/Kismet/Public/BlueprintEditorModule.h
@@ -9,6 +9,7 @@
 extern const FName BlueprintEditorAppName;
 
 class FBlueprintEditor;
+class UUserDefinedEnum;
 
 /**
  * Enum editor public interface
diff --git a/Engine/Source/Editor/KismetWidgets/Private/CreateBlueprintFromActorDialog.cpp b/Engine/Source/Editor/KismetWidgets/Private/CreateBlueprintFromActorDialog.cpp
--- a/Engine/Source/Editor/KismetWidgets/Private/CreateBlueprintFromActorDialog.cpp
+++ b/Engine/Source/Editor/KismetWidgets/Private/CreateBlueprintFromActorDialog.cpp
@@ -7,6 +7,7 @@
 #include "Editor/UnrealEd/Public/Kismet2/KismetEditorUtilities.h"
 #include "SNotificationList.h"
 #include "NotificationManager.h"
+#include "Engine/Selection.h"
 
 #define LOCTEXT_NAMESPACE "CreateBlueprintFromActorDialog"
 
diff --git a/Engine/Source/Editor/LandscapeEditor/Private/LandscapeEdMode.cpp b/Engine/Source/Editor/LandscapeEditor/Private/LandscapeEdMode.cpp
--- a/Engine/Source/Editor/LandscapeEditor/Private/LandscapeEdMode.cpp
+++ b/Engine/Source/Editor/LandscapeEditor/Private/LandscapeEdMode.cpp
@@ -32,6 +32,10 @@
 #include "LandscapeSplinesComponent.h"
 #include "Foliage/InstancedFoliageActor.h"
 #include "ComponentReregisterContext.h"
+#include "Engine/Selection.h"
+#include "LandscapeGizmoActiveActor.h"
+#include "EngineUtils.h"
+#include "Engine/Light.h"
 
 #define LOCTEXT_NAMESPACE "Landscape"
 
diff --git a/Engine/Source/Editor/LandscapeEditor/Private/LandscapeEdModeBrushes.cpp b/Engine/Source/Editor/LandscapeEditor/Private/LandscapeEdModeBrushes.cpp
--- a/Engine/Source/Editor/LandscapeEditor/Private/LandscapeEdModeBrushes.cpp
+++ b/Engine/Source/Editor/LandscapeEditor/Private/LandscapeEdModeBrushes.cpp
@@ -11,6 +11,7 @@
 #include "LandscapeHeightfieldCollisionComponent.h"
 
 #include "LevelUtils.h"
+#include "Materials/MaterialInstanceDynamic.h"
 
 // 
 // FLandscapeBrush
diff --git a/Engine/Source/Editor/LandscapeEditor/Private/LandscapeEdModeComponentTools.cpp b/Engine/Source/Editor/LandscapeEditor/Private/LandscapeEdModeComponentTools.cpp
--- a/Engine/Source/Editor/LandscapeEditor/Private/LandscapeEdModeComponentTools.cpp
+++ b/Engine/Source/Editor/LandscapeEditor/Private/LandscapeEdModeComponentTools.cpp
@@ -15,6 +15,7 @@
 #include "LandscapeEdModeTools.h"
 #include "Foliage/InstancedFoliageActor.h"
 #include "ComponentReregisterContext.h"
+#include "PhysicalMaterials/PhysicalMaterial.h"
 
 #define LOCTEXT_NAMESPACE "Landscape"
 
diff --git a/Engine/Source/Editor/LandscapeEditor/Private/LandscapeEdModePaintTools.cpp b/Engine/Source/Editor/LandscapeEditor/Private/LandscapeEdModePaintTools.cpp
--- a/Engine/Source/Editor/LandscapeEditor/Private/LandscapeEdModePaintTools.cpp
+++ b/Engine/Source/Editor/LandscapeEditor/Private/LandscapeEdModePaintTools.cpp
@@ -11,6 +11,7 @@
 #include "LandscapeEditorModule.h"
 #include "Editor/PropertyEditor/Public/PropertyEditorModule.h"
 #include "LandscapeEdModeTools.h"
+#include "Engine/StaticMesh.h"
 
 
 const int32 FNoiseParameter::Permutations[256] =
diff --git a/Engine/Source/Editor/LandscapeEditor/Private/LandscapeEdModeRampTool.cpp b/Engine/Source/Editor/LandscapeEditor/Private/LandscapeEdModeRampTool.cpp
--- a/Engine/Source/Editor/LandscapeEditor/Private/LandscapeEdModeRampTool.cpp
+++ b/Engine/Source/Editor/LandscapeEditor/Private/LandscapeEdModeRampTool.cpp
@@ -8,6 +8,7 @@
 #include "LandscapeDataAccess.h"
 #include "LandscapeHeightfieldCollisionComponent.h"
 #include "Raster.h"
+#include "AI/Navigation/NavigationSystem.h"
 
 #define LOCTEXT_NAMESPACE "Landscape"
 
diff --git a/Engine/Source/Editor/LandscapeEditor/Private/LandscapeEdModeSplineTools.cpp b/Engine/Source/Editor/LandscapeEditor/Private/LandscapeEdModeSplineTools.cpp
--- a/Engine/Source/Editor/LandscapeEditor/Private/LandscapeEdModeSplineTools.cpp
+++ b/Engine/Source/Editor/LandscapeEditor/Private/LandscapeEdModeSplineTools.cpp
@@ -17,6 +17,8 @@
 #include "LandscapeSplineControlPoint.h"
 #include "LandscapeSplineSegment.h"
 #include "ControlPointMeshComponent.h"
+#include "EditorUndoClient.h"
+#include "EngineUtils.h"
 
 
 #define LOCTEXT_NAMESPACE "Landscape"
diff --git a/Engine/Source/Editor/LandscapeEditor/Private/LandscapeEdModeTools.h b/Engine/Source/Editor/LandscapeEditor/Private/LandscapeEdModeTools.h
--- a/Engine/Source/Editor/LandscapeEditor/Private/LandscapeEdModeTools.h
+++ b/Engine/Source/Editor/LandscapeEditor/Private/LandscapeEdModeTools.h
@@ -5,6 +5,7 @@
 #include "Landscape.h"
 #include "LandscapeHeightfieldCollisionComponent.h"
 #include "Foliage/InstancedFoliageActor.h"
+#include "AI/Navigation/NavigationSystem.h"
 
 //
 //	FNoiseParameter - Perlin noise
diff --git a/Engine/Source/Editor/LandscapeEditor/Private/LandscapeEditorDetailCustomization_AlphaBrush.cpp b/Engine/Source/Editor/LandscapeEditor/Private/LandscapeEditorDetailCustomization_AlphaBrush.cpp
--- a/Engine/Source/Editor/LandscapeEditor/Private/LandscapeEditorDetailCustomization_AlphaBrush.cpp
+++ b/Engine/Source/Editor/LandscapeEditor/Private/LandscapeEditorDetailCustomization_AlphaBrush.cpp
@@ -17,6 +17,7 @@
 
 #include "Slate/SceneViewport.h"
 #include "AssetToolsModule.h"
+#include "CanvasTypes.h"
 
 #define LOCTEXT_NAMESPACE "LandscapeEditor.Brushes.Alpha"
 
diff --git a/Engine/Source/Editor/LandscapeEditor/Private/LandscapeEditorObject.cpp b/Engine/Source/Editor/LandscapeEditor/Private/LandscapeEditorObject.cpp
--- a/Engine/Source/Editor/LandscapeEditor/Private/LandscapeEditorObject.cpp
+++ b/Engine/Source/Editor/LandscapeEditor/Private/LandscapeEditorObject.cpp
@@ -6,6 +6,8 @@
 #include "ImageWrapper.h"
 
 #include "LandscapeMaterialInstanceConstant.h"
+#include "EngineUtils.h"
+#include "LandscapeGizmoActor.h"
 
 ULandscapeEditorObject::ULandscapeEditorObject(const FObjectInitializer& ObjectInitializer)
 	: Super(ObjectInitializer)
