Fixed issue with Undo when converting a particle module to seeded.

diff --git a/Engine/Source/Editor/Cascade/Private/Cascade.cpp b/Engine/Source/Editor/Cascade/Private/Cascade.cpp
--- a/Engine/Source/Editor/Cascade/Private/Cascade.cpp
+++ b/Engine/Source/Editor/Cascade/Private/Cascade.cpp
@@ -4669,15 +4669,6 @@ bool FCascade::ConvertModuleToSeeded(UParticleSystem* ParticleSystem, UParticleE
 		{
 			NewModule = CastChecked<UParticleModule>(StaticDuplicateObject(ConvertModule, ParticleSystem, TEXT("None"), RF_AllFlags, InSeededClass));
 
-			// Get all the curve objects and fixup their archetypes
-			// Otherwise they will be pointing to the original distribution
-			TArray<FParticleCurvePair> Curves;
-			NewModule->GetCurveObjects(Curves);
-			for (int32 CurveIdx = 0; CurveIdx < Curves.Num(); CurveIdx++)
-			{
-				FParticleCurvePair& Pair = Curves[CurveIdx];
-			}
-
 			// Since we used the non-randomseed module to create, this flag won't be set during construction...
 			NewModule->bSupportsRandomSeed = true;
 
@@ -4690,6 +4681,7 @@ bool FCascade::ConvertModuleToSeeded(UParticleSystem* ParticleSystem, UParticleE
 		}
 
 		// Now we have to replace all instances of the module
+		LODLevel->Modify();
 		LODLevel->Modules[InModuleIdx] = NewModule;
 		for (int32 SubLODIdx = LODIdx + 1; SubLODIdx < InEmitter->LODLevels.Num(); SubLODIdx++)
 		{
@@ -4699,6 +4691,7 @@ bool FCascade::ConvertModuleToSeeded(UParticleSystem* ParticleSystem, UParticleE
 			{
 				if (SubLODLevel->Modules[InModuleIdx] == ConvertModule)
 				{
+					SubLODLevel->Modify();
 					SubLODLevel->Modules[InModuleIdx] = NewModule;
 				}
 			}
@@ -4718,6 +4711,7 @@ bool FCascade::ConvertModuleToSeeded(UParticleSystem* ParticleSystem, UParticleE
 						UParticleModule* OtherModule = OtherLODLevel->Modules[OtherModuleIdx];
 						if (OtherModule == ConvertModule)
 						{
+							OtherLODLevel->Modify();
 							OtherLODLevel->Modules[OtherModuleIdx] = NewModule;
 							for (int32 OtherSubLODIdx = LODIdx + 1; OtherSubLODIdx < OtherEmitter->LODLevels.Num(); OtherSubLODIdx++)
 							{
@@ -4727,6 +4721,7 @@ bool FCascade::ConvertModuleToSeeded(UParticleSystem* ParticleSystem, UParticleE
 								{
 									if (OtherSubLODLevel->Modules[InModuleIdx] == ConvertModule)
 									{
+										OtherSubLODLevel->Modify();
 										OtherSubLODLevel->Modules[InModuleIdx] = NewModule;
 									}
 								}
