Fixed non-unity

diff --git a/Engine/Source/Editor/Cascade/Private/CascadePreviewViewportClient.cpp b/Engine/Source/Editor/Cascade/Private/CascadePreviewViewportClient.cpp
--- a/Engine/Source/Editor/Cascade/Private/CascadePreviewViewportClient.cpp
+++ b/Engine/Source/Editor/Cascade/Private/CascadePreviewViewportClient.cpp
@@ -6,6 +6,7 @@
 #include "PreviewScene.h"
 #include "ImageUtils.h"
 #include "CascadePreviewViewportClient.h"
+#include "SCascadePreviewViewport.h"
 #include "Particles/Spawn/ParticleModuleSpawn.h"
 #include "Particles/TypeData/ParticleModuleTypeDataGpu.h"
 #include "Particles/VectorField/ParticleModuleVectorFieldLocal.h"
