Cascade: Fix light preview visualization not showing up in Cascade [UE-10906]

diff --git a/Engine/Source/Editor/Cascade/Private/CascadePreviewViewportClient.cpp b/Engine/Source/Editor/Cascade/Private/CascadePreviewViewportClient.cpp
--- a/Engine/Source/Editor/Cascade/Private/CascadePreviewViewportClient.cpp
+++ b/Engine/Source/Editor/Cascade/Private/CascadePreviewViewportClient.cpp
@@ -32,7 +32,7 @@ FCascadeEdPreviewViewportClient::FCascadeEdPreviewViewportClient(TWeakPtr<FCasca
 	, VectorFieldHitproxyInfo(0)
 	, LightRotSpeed(0.22f)
 {
-
+	PreviewScene = &CascadePreviewScene;
 	check(CascadePtr.IsValid() && EditorViewportWidget.IsValid());
 
 	UParticleSystem* ParticleSystem = CascadePtr.Pin()->GetParticleSystem();
@@ -528,6 +528,9 @@ void FCascadeEdPreviewViewportClient::Draw(const FSceneView* View, FPrimitiveDra
 			}
 		}
 	}
+
+	// Draw the preview scene light visualization
+	DrawPreviewLightVisualization(View, PDI);
 }
 
 bool FCascadeEdPreviewViewportClient::InputKey(FViewport* Viewport, int32 ControllerId, FKey Key, EInputEvent Event, float AmountDepressed, bool Gamepad)
@@ -600,23 +603,8 @@ static TAutoConsoleVariable<float> CVarCascadeScaleSpeed(TEXT("CascadeScaleSpeed
 bool FCascadeEdPreviewViewportClient::InputAxis(FViewport* Viewport, int32 ControllerId, FKey Key, float Delta, float DeltaTime, int32 NumSamples, bool bGamepad)
 {
 	bool bHandled = false;
-	const bool bLightMoveDown = Viewport->KeyState(EKeys::L);
-	if(bLightMoveDown)
-	{
-		FRotator LightDir = CascadePreviewScene.GetLightDirection();
-		// Look at which axis is being dragged and by how much
-		const float DragDeltaX = (Key == EKeys::MouseX) ? Delta : 0.f;
-		const float DragDeltaY = (Key == EKeys::MouseY) ? Delta : 0.f;
-
-		LightDir.Yaw += -DragDeltaX * LightRotSpeed;
-		LightDir.Pitch += -DragDeltaY * LightRotSpeed;
 
-		CascadePreviewScene.SetLightDirection(LightDir);
-
-		Viewport->Invalidate();
-		bHandled = true;
-	}
-	else if(bManipulatingVectorField)
+	if (bManipulatingVectorField)
 	{
 		UParticleModuleVectorFieldLocal* VectorFieldModule = Cast<UParticleModuleVectorFieldLocal>(CascadePtr.Pin()->GetSelectedModule());
 		if (VectorFieldModule)
