Unify Slate's FInputGesture and Engine's FInputChord as FInputChord defined in Slate

diff --git a/Engine/Plugins/2D/Paper2D/Source/Paper2DEditor/Private/FlipbookEditor/FlipbookEditorCommands.cpp b/Engine/Plugins/2D/Paper2D/Source/Paper2DEditor/Private/FlipbookEditor/FlipbookEditorCommands.cpp
--- a/Engine/Plugins/2D/Paper2D/Source/Paper2DEditor/Private/FlipbookEditor/FlipbookEditorCommands.cpp
+++ b/Engine/Plugins/2D/Paper2D/Source/Paper2DEditor/Private/FlipbookEditor/FlipbookEditorCommands.cpp
@@ -10,20 +10,20 @@
 
 void FFlipbookEditorCommands::RegisterCommands()
 {
-	UI_COMMAND(SetShowGrid, "Grid", "Displays the viewport grid.", EUserInterfaceActionType::ToggleButton, FInputGesture());
-	UI_COMMAND(SetShowBounds, "Bounds", "Toggles display of the bounds of the static mesh.", EUserInterfaceActionType::ToggleButton, FInputGesture());
-	UI_COMMAND(SetShowCollision, "Collision", "Toggles display of the simplified collision mesh of the static mesh, if one has been assigned.", EUserInterfaceActionType::ToggleButton, FInputGesture(EKeys::C, EModifierKey::Alt));
+	UI_COMMAND(SetShowGrid, "Grid", "Displays the viewport grid.", EUserInterfaceActionType::ToggleButton, FInputChord());
+	UI_COMMAND(SetShowBounds, "Bounds", "Toggles display of the bounds of the static mesh.", EUserInterfaceActionType::ToggleButton, FInputChord());
+	UI_COMMAND(SetShowCollision, "Collision", "Toggles display of the simplified collision mesh of the static mesh, if one has been assigned.", EUserInterfaceActionType::ToggleButton, FInputChord(EKeys::C, EModifierKey::Alt));
 
-	UI_COMMAND(SetShowPivot, "Show Pivot", "Display the pivot location of the static mesh.", EUserInterfaceActionType::ToggleButton, FInputGesture());
-	UI_COMMAND(SetShowSockets, "Sockets", "Displays the flipbook sockets.", EUserInterfaceActionType::ToggleButton, FInputGesture());
+	UI_COMMAND(SetShowPivot, "Show Pivot", "Display the pivot location of the static mesh.", EUserInterfaceActionType::ToggleButton, FInputChord());
+	UI_COMMAND(SetShowSockets, "Sockets", "Displays the flipbook sockets.", EUserInterfaceActionType::ToggleButton, FInputChord());
 
-	UI_COMMAND(AddNewFrame, "Add Key Frame", "Adds a new key frame to the flipbook.", EUserInterfaceActionType::Button, FInputGesture());
-	UI_COMMAND(AddNewFrameBefore, "Insert Key Frame Before", "Adds a new key frame to the flipbook before the selection.", EUserInterfaceActionType::Button, FInputGesture());
-	UI_COMMAND(AddNewFrameAfter, "Insert Key Frame After", "Adds a new key frame to the flipbook after the selection.", EUserInterfaceActionType::Button, FInputGesture());
+	UI_COMMAND(AddNewFrame, "Add Key Frame", "Adds a new key frame to the flipbook.", EUserInterfaceActionType::Button, FInputChord());
+	UI_COMMAND(AddNewFrameBefore, "Insert Key Frame Before", "Adds a new key frame to the flipbook before the selection.", EUserInterfaceActionType::Button, FInputChord());
+	UI_COMMAND(AddNewFrameAfter, "Insert Key Frame After", "Adds a new key frame to the flipbook after the selection.", EUserInterfaceActionType::Button, FInputChord());
 
-	UI_COMMAND(PickNewSpriteFrame, "Pick New Sprite", "Picks a new sprite for this key frame.", EUserInterfaceActionType::Button, FInputGesture());
-	UI_COMMAND(EditSpriteFrame, "Edit Sprite", "Opens the sprite for this key frame in the Sprite Editor.", EUserInterfaceActionType::Button, FInputGesture());
-	UI_COMMAND(ShowInContentBrowser, "Show in Content Browser", "Shows the sprite for this key frame in the Content Browser.", EUserInterfaceActionType::Button, FInputGesture());
+	UI_COMMAND(PickNewSpriteFrame, "Pick New Sprite", "Picks a new sprite for this key frame.", EUserInterfaceActionType::Button, FInputChord());
+	UI_COMMAND(EditSpriteFrame, "Edit Sprite", "Opens the sprite for this key frame in the Sprite Editor.", EUserInterfaceActionType::Button, FInputChord());
+	UI_COMMAND(ShowInContentBrowser, "Show in Content Browser", "Shows the sprite for this key frame in the Content Browser.", EUserInterfaceActionType::Button, FInputChord());
 }
 
 #undef LOCTEXT_NAMESPACE
\ No newline at end of file
diff --git a/Engine/Plugins/2D/Paper2D/Source/Paper2DEditor/Private/SpriteEditor/SpriteEditorCommands.cpp b/Engine/Plugins/2D/Paper2D/Source/Paper2DEditor/Private/SpriteEditor/SpriteEditorCommands.cpp
--- a/Engine/Plugins/2D/Paper2D/Source/Paper2DEditor/Private/SpriteEditor/SpriteEditorCommands.cpp
+++ b/Engine/Plugins/2D/Paper2D/Source/Paper2DEditor/Private/SpriteEditor/SpriteEditorCommands.cpp
@@ -11,32 +11,32 @@
 void FSpriteEditorCommands::RegisterCommands()
 {
 	// Show toggles
-	UI_COMMAND(SetShowGrid, "Grid", "Displays the viewport grid.", EUserInterfaceActionType::ToggleButton, FInputGesture());
-	UI_COMMAND(SetShowSourceTexture, "Src Tex", "Toggles display of the source texture (useful when it is an atlas).", EUserInterfaceActionType::ToggleButton, FInputGesture());
-	UI_COMMAND(SetShowBounds, "Bounds", "Toggles display of the bounds of the static mesh.", EUserInterfaceActionType::ToggleButton, FInputGesture());
-	UI_COMMAND(SetShowCollision, "Collision", "Toggles display of the simplified collision mesh of the static mesh, if one has been assigned.", EUserInterfaceActionType::ToggleButton, FInputGesture(EKeys::C, EModifierKey::Alt));
+	UI_COMMAND(SetShowGrid, "Grid", "Displays the viewport grid.", EUserInterfaceActionType::ToggleButton, FInputChord());
+	UI_COMMAND(SetShowSourceTexture, "Src Tex", "Toggles display of the source texture (useful when it is an atlas).", EUserInterfaceActionType::ToggleButton, FInputChord());
+	UI_COMMAND(SetShowBounds, "Bounds", "Toggles display of the bounds of the static mesh.", EUserInterfaceActionType::ToggleButton, FInputChord());
+	UI_COMMAND(SetShowCollision, "Collision", "Toggles display of the simplified collision mesh of the static mesh, if one has been assigned.", EUserInterfaceActionType::ToggleButton, FInputChord(EKeys::C, EModifierKey::Alt));
 
-	UI_COMMAND(SetShowSockets, "Sockets", "Displays the sprite sockets.", EUserInterfaceActionType::ToggleButton, FInputGesture());
+	UI_COMMAND(SetShowSockets, "Sockets", "Displays the sprite sockets.", EUserInterfaceActionType::ToggleButton, FInputChord());
 
-	UI_COMMAND(SetShowNormals, "Normals", "Toggles display of vertex normals in the Preview Pane.", EUserInterfaceActionType::ToggleButton, FInputGesture());
-	UI_COMMAND(SetShowPivot, "Pivot", "Display the pivot location of the sprite.", EUserInterfaceActionType::ToggleButton, FInputGesture());
-	UI_COMMAND(SetShowMeshEdges, "Mesh Edges", "Overlays the mesh edges on top of the view.", EUserInterfaceActionType::ToggleButton, FInputGesture());
+	UI_COMMAND(SetShowNormals, "Normals", "Toggles display of vertex normals in the Preview Pane.", EUserInterfaceActionType::ToggleButton, FInputChord());
+	UI_COMMAND(SetShowPivot, "Pivot", "Display the pivot location of the sprite.", EUserInterfaceActionType::ToggleButton, FInputChord());
+	UI_COMMAND(SetShowMeshEdges, "Mesh Edges", "Overlays the mesh edges on top of the view.", EUserInterfaceActionType::ToggleButton, FInputChord());
 
 	// Editing modes
-	UI_COMMAND(EnterViewMode, "View", "View the sprite.", EUserInterfaceActionType::ToggleButton, FInputGesture());
-	UI_COMMAND(EnterSourceRegionEditMode, "Edit Source Region", "Edit the sprite source region.", EUserInterfaceActionType::ToggleButton, FInputGesture());
-	UI_COMMAND(EnterCollisionEditMode, "Edit Collision", "Edit the collision geometry.", EUserInterfaceActionType::ToggleButton, FInputGesture());
-	UI_COMMAND(EnterRenderingEditMode, "Edit RenderGeom", "Edit the rendering geometry (useful to reduce overdraw).", EUserInterfaceActionType::ToggleButton, FInputGesture());
-	UI_COMMAND(EnterAddSpriteMode, "Add Sprite", "Create new sprites quickly.", EUserInterfaceActionType::ToggleButton, FInputGesture());
+	UI_COMMAND(EnterViewMode, "View", "View the sprite.", EUserInterfaceActionType::ToggleButton, FInputChord());
+	UI_COMMAND(EnterSourceRegionEditMode, "Edit Source Region", "Edit the sprite source region.", EUserInterfaceActionType::ToggleButton, FInputChord());
+	UI_COMMAND(EnterCollisionEditMode, "Edit Collision", "Edit the collision geometry.", EUserInterfaceActionType::ToggleButton, FInputChord());
+	UI_COMMAND(EnterRenderingEditMode, "Edit RenderGeom", "Edit the rendering geometry (useful to reduce overdraw).", EUserInterfaceActionType::ToggleButton, FInputChord());
+	UI_COMMAND(EnterAddSpriteMode, "Add Sprite", "Create new sprites quickly.", EUserInterfaceActionType::ToggleButton, FInputChord());
 
 	// Misc. actions
-	UI_COMMAND(FocusOnSprite, "Focus on sprite", "Centers and zooms the view to focus on the current sprite.", EUserInterfaceActionType::Button, FInputGesture(EKeys::Home, EModifierKey::None));
+	UI_COMMAND(FocusOnSprite, "Focus on sprite", "Centers and zooms the view to focus on the current sprite.", EUserInterfaceActionType::Button, FInputChord(EKeys::Home, EModifierKey::None));
 
 	// Geometry editing commands
-	UI_COMMAND(DeleteSelection, "Delete", "Delete the selection.", EUserInterfaceActionType::Button, FInputGesture(EKeys::Platform_Delete, EModifierKey::None));
-	UI_COMMAND(SplitEdge, "Split", "Split edge.", EUserInterfaceActionType::Button, FInputGesture(EKeys::Insert, EModifierKey::None));
-	UI_COMMAND(ToggleAddPolygonMode, "Add Polygon", "Adds a new polygon.", EUserInterfaceActionType::ToggleButton, FInputGesture());
-	UI_COMMAND(SnapAllVertices, "Snap to pixel grid", "Snaps all vertices to the pixel grid.", EUserInterfaceActionType::Button, FInputGesture());
+	UI_COMMAND(DeleteSelection, "Delete", "Delete the selection.", EUserInterfaceActionType::Button, FInputChord(EKeys::Platform_Delete, EModifierKey::None));
+	UI_COMMAND(SplitEdge, "Split", "Split edge.", EUserInterfaceActionType::Button, FInputChord(EKeys::Insert, EModifierKey::None));
+	UI_COMMAND(ToggleAddPolygonMode, "Add Polygon", "Adds a new polygon.", EUserInterfaceActionType::ToggleButton, FInputChord());
+	UI_COMMAND(SnapAllVertices, "Snap to pixel grid", "Snaps all vertices to the pixel grid.", EUserInterfaceActionType::Button, FInputChord());
 }
 
 #undef LOCTEXT_NAMESPACE
\ No newline at end of file
diff --git a/Engine/Plugins/2D/Paper2D/Source/Paper2DEditor/Private/TileMapEditing/TileMapEditorCommands.cpp b/Engine/Plugins/2D/Paper2D/Source/Paper2DEditor/Private/TileMapEditing/TileMapEditorCommands.cpp
--- a/Engine/Plugins/2D/Paper2D/Source/Paper2DEditor/Private/TileMapEditing/TileMapEditorCommands.cpp
+++ b/Engine/Plugins/2D/Paper2D/Source/Paper2DEditor/Private/TileMapEditing/TileMapEditorCommands.cpp
@@ -10,34 +10,34 @@
 
 void FTileMapEditorCommands::RegisterCommands()
 {
-	UI_COMMAND(EnterTileMapEditMode, "Enable Tile Map Mode", "Enables Tile Map editing mode", EUserInterfaceActionType::Button, FInputGesture());
+	UI_COMMAND(EnterTileMapEditMode, "Enable Tile Map Mode", "Enables Tile Map editing mode", EUserInterfaceActionType::Button, FInputChord());
 
-	UI_COMMAND(SelectPaintTool, "Paint", "Paint", EUserInterfaceActionType::RadioButton, FInputGesture(EKeys::B));
-	UI_COMMAND(SelectEraserTool, "Eraser", "Eraser", EUserInterfaceActionType::RadioButton, FInputGesture(EKeys::E));
-	UI_COMMAND(SelectFillTool, "Fill", "Paint Bucket", EUserInterfaceActionType::RadioButton, FInputGesture(EKeys::F));
+	UI_COMMAND(SelectPaintTool, "Paint", "Paint", EUserInterfaceActionType::RadioButton, FInputChord(EKeys::B));
+	UI_COMMAND(SelectEraserTool, "Eraser", "Eraser", EUserInterfaceActionType::RadioButton, FInputChord(EKeys::E));
+	UI_COMMAND(SelectFillTool, "Fill", "Paint Bucket", EUserInterfaceActionType::RadioButton, FInputChord(EKeys::F));
 
 	// Show toggles
-	UI_COMMAND(SetShowCollision, "Collision", "Toggles display of the simplified collision mesh of the static mesh, if one has been assigned.", EUserInterfaceActionType::ToggleButton, FInputGesture(EKeys::C, EModifierKey::Alt));
+	UI_COMMAND(SetShowCollision, "Collision", "Toggles display of the simplified collision mesh of the static mesh, if one has been assigned.", EUserInterfaceActionType::ToggleButton, FInputChord(EKeys::C, EModifierKey::Alt));
 
-	UI_COMMAND(SetShowPivot, "Pivot", "Display the pivot location of the sprite.", EUserInterfaceActionType::ToggleButton, FInputGesture());
+	UI_COMMAND(SetShowPivot, "Pivot", "Display the pivot location of the sprite.", EUserInterfaceActionType::ToggleButton, FInputChord());
 
 	// Misc. actions
-	UI_COMMAND(FocusOnTileMap, "Focus on tile map", "Centers and zooms the view to focus on the tile map.", EUserInterfaceActionType::Button, FInputGesture(EKeys::Home, EModifierKey::None));
+	UI_COMMAND(FocusOnTileMap, "Focus on tile map", "Centers and zooms the view to focus on the tile map.", EUserInterfaceActionType::Button, FInputChord(EKeys::Home, EModifierKey::None));
 
 	// Selection actions
-	UI_COMMAND(FlipSelectionHorizontally, "Flip selection horizontally", "Flips the selection horizontally", EUserInterfaceActionType::Button, FInputGesture(EKeys::X));
-	UI_COMMAND(FlipSelectionVertically, "Flip selection vertically", "Flips the selection vertically", EUserInterfaceActionType::Button, FInputGesture(EKeys::Y));
-	UI_COMMAND(RotateSelectionCW, "Rotate selection clockwise", "Rotates the selection clockwise", EUserInterfaceActionType::Button, FInputGesture(EKeys::Z));
-	UI_COMMAND(RotateSelectionCCW, "Rotate selection counterclockwise", "Rotates the selection counterclockwise", EUserInterfaceActionType::Button, FInputGesture(EKeys::Z, EModifierKey::Shift));
+	UI_COMMAND(FlipSelectionHorizontally, "Flip selection horizontally", "Flips the selection horizontally", EUserInterfaceActionType::Button, FInputChord(EKeys::X));
+	UI_COMMAND(FlipSelectionVertically, "Flip selection vertically", "Flips the selection vertically", EUserInterfaceActionType::Button, FInputChord(EKeys::Y));
+	UI_COMMAND(RotateSelectionCW, "Rotate selection clockwise", "Rotates the selection clockwise", EUserInterfaceActionType::Button, FInputChord(EKeys::Z));
+	UI_COMMAND(RotateSelectionCCW, "Rotate selection counterclockwise", "Rotates the selection counterclockwise", EUserInterfaceActionType::Button, FInputChord(EKeys::Z, EModifierKey::Shift));
 
 	// Layer actions
-	UI_COMMAND(AddNewLayerAbove, "Add new layer", "Add a new layer above the current layer", EUserInterfaceActionType::Button, FInputGesture());
-	UI_COMMAND(AddNewLayerBelow, "Add new layer below", "Add a new layer below the current layer", EUserInterfaceActionType::Button, FInputGesture());
-	UI_COMMAND(DeleteLayer, "Delete layer", "Delete the current layer", EUserInterfaceActionType::Button, FInputGesture());
-	UI_COMMAND(DuplicateLayer, "Duplicate layer", "Duplicate the current layer", EUserInterfaceActionType::Button, FInputGesture());
-	UI_COMMAND(MergeLayerDown, "Merge layer down", "Merge the current layer down to the layer below it", EUserInterfaceActionType::Button, FInputGesture());
-	UI_COMMAND(MoveLayerUp, "Move layer up", "Move layer up", EUserInterfaceActionType::Button, FInputGesture());
-	UI_COMMAND(MoveLayerDown, "Move layer down", "Move layer down", EUserInterfaceActionType::Button, FInputGesture());
+	UI_COMMAND(AddNewLayerAbove, "Add new layer", "Add a new layer above the current layer", EUserInterfaceActionType::Button, FInputChord());
+	UI_COMMAND(AddNewLayerBelow, "Add new layer below", "Add a new layer below the current layer", EUserInterfaceActionType::Button, FInputChord());
+	UI_COMMAND(DeleteLayer, "Delete layer", "Delete the current layer", EUserInterfaceActionType::Button, FInputChord());
+	UI_COMMAND(DuplicateLayer, "Duplicate layer", "Duplicate the current layer", EUserInterfaceActionType::Button, FInputChord());
+	UI_COMMAND(MergeLayerDown, "Merge layer down", "Merge the current layer down to the layer below it", EUserInterfaceActionType::Button, FInputChord());
+	UI_COMMAND(MoveLayerUp, "Move layer up", "Move layer up", EUserInterfaceActionType::Button, FInputChord());
+	UI_COMMAND(MoveLayerDown, "Move layer down", "Move layer down", EUserInterfaceActionType::Button, FInputChord());
 }
 
 #undef LOCTEXT_NAMESPACE
\ No newline at end of file
diff --git a/Engine/Plugins/Editor/EpicSurvey/Source/EpicSurvey/Private/EpicSurveyCommands.cpp b/Engine/Plugins/Editor/EpicSurvey/Source/EpicSurvey/Private/EpicSurveyCommands.cpp
--- a/Engine/Plugins/Editor/EpicSurvey/Source/EpicSurvey/Private/EpicSurveyCommands.cpp
+++ b/Engine/Plugins/Editor/EpicSurvey/Source/EpicSurvey/Private/EpicSurveyCommands.cpp
@@ -17,7 +17,7 @@ FEpicSurveyCommands::FEpicSurveyCommands()
 
 void FEpicSurveyCommands::RegisterCommands()
 {
-	UI_COMMAND( OpenEpicSurvey, "Open Epic Survey...", "Allows you to send feedback to Epic to help improve UE4.", EUserInterfaceActionType::Button, FInputGesture() );
+	UI_COMMAND( OpenEpicSurvey, "Open Epic Survey...", "Allows you to send feedback to Epic to help improve UE4.", EUserInterfaceActionType::Button, FInputChord() );
 }
 
 #undef LOCTEXT_NAMESPACE
\ No newline at end of file
diff --git a/Engine/Plugins/Media/MediaPlayerEditor/Source/MediaPlayerEditor/Private/Models/MediaPlayerEditorCommands.cpp b/Engine/Plugins/Media/MediaPlayerEditor/Source/MediaPlayerEditor/Private/Models/MediaPlayerEditorCommands.cpp
--- a/Engine/Plugins/Media/MediaPlayerEditor/Source/MediaPlayerEditor/Private/Models/MediaPlayerEditorCommands.cpp
+++ b/Engine/Plugins/Media/MediaPlayerEditor/Source/MediaPlayerEditor/Private/Models/MediaPlayerEditorCommands.cpp
@@ -6,11 +6,11 @@
 
 void FMediaPlayerEditorCommands::RegisterCommands()
 {
-	UI_COMMAND(ForwardMedia, "Forward", "Fast forwards media playback", EUserInterfaceActionType::Button, FInputGesture());
-	UI_COMMAND(PauseMedia, "Pause", "Pause media playback", EUserInterfaceActionType::Button, FInputGesture());
-	UI_COMMAND(PlayMedia, "Play", "Start media playback", EUserInterfaceActionType::Button, FInputGesture());
-	UI_COMMAND(ReverseMedia, "Reverse", "Reverses media playback", EUserInterfaceActionType::Button, FInputGesture());
-	UI_COMMAND(RewindMedia, "Rewind", "Rewinds the media to the beginning", EUserInterfaceActionType::Button, FInputGesture());
+	UI_COMMAND(ForwardMedia, "Forward", "Fast forwards media playback", EUserInterfaceActionType::Button, FInputChord());
+	UI_COMMAND(PauseMedia, "Pause", "Pause media playback", EUserInterfaceActionType::Button, FInputChord());
+	UI_COMMAND(PlayMedia, "Play", "Start media playback", EUserInterfaceActionType::Button, FInputChord());
+	UI_COMMAND(ReverseMedia, "Reverse", "Reverses media playback", EUserInterfaceActionType::Button, FInputChord());
+	UI_COMMAND(RewindMedia, "Rewind", "Rewinds the media to the beginning", EUserInterfaceActionType::Button, FInputChord());
 }
 
 #undef LOCTEXT_NAMESPACE
\ No newline at end of file
diff --git a/Engine/Plugins/Messaging/MessagingDebugger/Source/MessagingDebugger/Private/Models/MessagingDebuggerCommands.h b/Engine/Plugins/Messaging/MessagingDebugger/Source/MessagingDebugger/Private/Models/MessagingDebuggerCommands.h
--- a/Engine/Plugins/Messaging/MessagingDebugger/Source/MessagingDebugger/Private/Models/MessagingDebuggerCommands.h
+++ b/Engine/Plugins/Messaging/MessagingDebugger/Source/MessagingDebugger/Private/Models/MessagingDebuggerCommands.h
@@ -22,12 +22,12 @@ public:
 
 	virtual void RegisterCommands() override
 	{
-		UI_COMMAND(BreakDebugger, "Break", "Break the debugger at the next message", EUserInterfaceActionType::Button, FInputGesture());
-		UI_COMMAND(ClearHistory, "Clear History", "Clears the message history list", EUserInterfaceActionType::Button, FInputGesture());
-		UI_COMMAND(ContinueDebugger, "Continue", "Continues debugging", EUserInterfaceActionType::Button, FInputGesture());
-		UI_COMMAND(StartDebugger, "Start", "Start the debugger", EUserInterfaceActionType::Button, FInputGesture());
-		UI_COMMAND(StepDebugger, "Step", "Step over the current message", EUserInterfaceActionType::Button, FInputGesture());
-		UI_COMMAND(StopDebugger, "Stop", "Stop the debugger", EUserInterfaceActionType::Button, FInputGesture());
+		UI_COMMAND(BreakDebugger, "Break", "Break the debugger at the next message", EUserInterfaceActionType::Button, FInputChord());
+		UI_COMMAND(ClearHistory, "Clear History", "Clears the message history list", EUserInterfaceActionType::Button, FInputChord());
+		UI_COMMAND(ContinueDebugger, "Continue", "Continues debugging", EUserInterfaceActionType::Button, FInputChord());
+		UI_COMMAND(StartDebugger, "Start", "Start the debugger", EUserInterfaceActionType::Button, FInputChord());
+		UI_COMMAND(StepDebugger, "Step", "Step over the current message", EUserInterfaceActionType::Button, FInputChord());
+		UI_COMMAND(StopDebugger, "Stop", "Stop the debugger", EUserInterfaceActionType::Button, FInputChord());
 	}
 
 public:
diff --git a/Engine/Source/Developer/AutomationWindow/Private/SAutomationWindow.cpp b/Engine/Source/Developer/AutomationWindow/Private/SAutomationWindow.cpp
--- a/Engine/Source/Developer/AutomationWindow/Private/SAutomationWindow.cpp
+++ b/Engine/Source/Developer/AutomationWindow/Private/SAutomationWindow.cpp
@@ -30,13 +30,13 @@ public:
 
 	virtual void RegisterCommands() override
 	{
-		UI_COMMAND( RefreshTests, "Refresh Tests", "Refresh Tests", EUserInterfaceActionType::Button, FInputGesture() );
-		UI_COMMAND( FindWorkers, "Find Workers", "Find Workers", EUserInterfaceActionType::Button, FInputGesture() );
-		UI_COMMAND( ErrorFilter, "Errors", "Toggle Error Filter", EUserInterfaceActionType::ToggleButton, FInputGesture() );
-		UI_COMMAND( WarningFilter, "Warnings", "Toggle Warning Filter", EUserInterfaceActionType::ToggleButton, FInputGesture() );
-		UI_COMMAND( SmokeTestFilter, "Smoke Tests", "Toggle Smoke Test Filter", EUserInterfaceActionType::ToggleButton, FInputGesture() );
-		UI_COMMAND( DeveloperDirectoryContent, "Dev Content", "Developer Directory Content Filter (when enabled, developer directories are also included)", EUserInterfaceActionType::ToggleButton, FInputGesture() );
-		UI_COMMAND( VisualCommandlet, "Vis Cmdlet", "Visual Commandlet Filter", EUserInterfaceActionType::ToggleButton, FInputGesture() );
+		UI_COMMAND( RefreshTests, "Refresh Tests", "Refresh Tests", EUserInterfaceActionType::Button, FInputChord() );
+		UI_COMMAND( FindWorkers, "Find Workers", "Find Workers", EUserInterfaceActionType::Button, FInputChord() );
+		UI_COMMAND( ErrorFilter, "Errors", "Toggle Error Filter", EUserInterfaceActionType::ToggleButton, FInputChord() );
+		UI_COMMAND( WarningFilter, "Warnings", "Toggle Warning Filter", EUserInterfaceActionType::ToggleButton, FInputChord() );
+		UI_COMMAND( SmokeTestFilter, "Smoke Tests", "Toggle Smoke Test Filter", EUserInterfaceActionType::ToggleButton, FInputChord() );
+		UI_COMMAND( DeveloperDirectoryContent, "Dev Content", "Developer Directory Content Filter (when enabled, developer directories are also included)", EUserInterfaceActionType::ToggleButton, FInputChord() );
+		UI_COMMAND( VisualCommandlet, "Vis Cmdlet", "Visual Commandlet Filter", EUserInterfaceActionType::ToggleButton, FInputChord() );
 	}
 public:
 	TSharedPtr<FUICommandInfo> RefreshTests;
diff --git a/Engine/Source/Developer/DeviceManager/Private/Models/DeviceDetailsCommands.h b/Engine/Source/Developer/DeviceManager/Private/Models/DeviceDetailsCommands.h
--- a/Engine/Source/Developer/DeviceManager/Private/Models/DeviceDetailsCommands.h
+++ b/Engine/Source/Developer/DeviceManager/Private/Models/DeviceDetailsCommands.h
@@ -29,18 +29,18 @@ public:
 
 	virtual void RegisterCommands( ) override
 	{
-		UI_COMMAND(Claim, "Claim", "Claim the device", EUserInterfaceActionType::Button, FInputGesture());
-		UI_COMMAND(Release, "Release", "Release the device", EUserInterfaceActionType::Button, FInputGesture());
-		UI_COMMAND(Remove, "Remove", "Remove the device", EUserInterfaceActionType::Button, FInputGesture());
-		UI_COMMAND(Share, "Share", "Share the device with other users", EUserInterfaceActionType::ToggleButton, FInputGesture());
+		UI_COMMAND(Claim, "Claim", "Claim the device", EUserInterfaceActionType::Button, FInputChord());
+		UI_COMMAND(Release, "Release", "Release the device", EUserInterfaceActionType::Button, FInputChord());
+		UI_COMMAND(Remove, "Remove", "Remove the device", EUserInterfaceActionType::Button, FInputChord());
+		UI_COMMAND(Share, "Share", "Share the device with other users", EUserInterfaceActionType::ToggleButton, FInputChord());
 
-		UI_COMMAND(Connect, "Connect", "Connect to the device", EUserInterfaceActionType::Button, FInputGesture());
-		UI_COMMAND(Disconnect, "Disconnect", "Disconnect from the device", EUserInterfaceActionType::Button, FInputGesture());
+		UI_COMMAND(Connect, "Connect", "Connect to the device", EUserInterfaceActionType::Button, FInputChord());
+		UI_COMMAND(Disconnect, "Disconnect", "Disconnect from the device", EUserInterfaceActionType::Button, FInputChord());
 
-		UI_COMMAND(PowerOff, "Power Off", "Power off the device", EUserInterfaceActionType::Button, FInputGesture());
-		UI_COMMAND(PowerOffForce, "Power Off (Force)", "Power off the device (forcefully close any running applications)", EUserInterfaceActionType::Button, FInputGesture());
-		UI_COMMAND(PowerOn, "Power On", "Power on the device", EUserInterfaceActionType::Button, FInputGesture());
-		UI_COMMAND(Reboot, "Reboot", "Reboot the device", EUserInterfaceActionType::Button, FInputGesture());
+		UI_COMMAND(PowerOff, "Power Off", "Power off the device", EUserInterfaceActionType::Button, FInputChord());
+		UI_COMMAND(PowerOffForce, "Power Off (Force)", "Power off the device (forcefully close any running applications)", EUserInterfaceActionType::Button, FInputChord());
+		UI_COMMAND(PowerOn, "Power On", "Power on the device", EUserInterfaceActionType::Button, FInputChord());
+		UI_COMMAND(Reboot, "Reboot", "Reboot the device", EUserInterfaceActionType::Button, FInputChord());
 	}
 
 public:
diff --git a/Engine/Source/Developer/LogVisualizer/Private/VisualLoggerCommands.h b/Engine/Source/Developer/LogVisualizer/Private/VisualLoggerCommands.h
--- a/Engine/Source/Developer/LogVisualizer/Private/VisualLoggerCommands.h
+++ b/Engine/Source/Developer/LogVisualizer/Private/VisualLoggerCommands.h
@@ -20,19 +20,19 @@ public:
 
 	virtual void RegisterCommands() override
 	{
-		UI_COMMAND(StartRecording, "Start", "Start the debugger", EUserInterfaceActionType::Button, FInputGesture());
-		UI_COMMAND(StopRecording, "Stop", "Step over the current message", EUserInterfaceActionType::Button, FInputGesture());
-		UI_COMMAND(Pause, "Pause", "Stop the debugger", EUserInterfaceActionType::Button, FInputGesture());
-		UI_COMMAND(Resume, "Resume", "Stop the debugger", EUserInterfaceActionType::Button, FInputGesture());
-		UI_COMMAND(FreeCamera, "FreeCamera", "Enable free camera", EUserInterfaceActionType::ToggleButton, FInputGesture());
-		UI_COMMAND(LoadFromVLog, "Load", "Load external vlogs", EUserInterfaceActionType::Button, FInputGesture());
-		UI_COMMAND(SaveToVLog, "Save", "Save selected data to vlog file", EUserInterfaceActionType::Button, FInputGesture());
-		UI_COMMAND(ToggleGraphs, "ToggleGraphs", "Toggle graphs visualization on/off", EUserInterfaceActionType::ToggleButton, FInputGesture());
-		UI_COMMAND(ResetData, "Clear", "Clear all data", EUserInterfaceActionType::Button, FInputGesture());
-		UI_COMMAND(HideEmpty, "HideEmpty", "Hide logs without text information", EUserInterfaceActionType::ToggleButton, FInputGesture());
-
-		UI_COMMAND(MoveCursorLeft, "MoveCursorLeft", "Move cursor left on timeline", EUserInterfaceActionType::Button, FInputGesture(EKeys::Left));
-		UI_COMMAND(MoveCursorRight, "MoveCursorLeft", "Move cursor right on timeline", EUserInterfaceActionType::Button, FInputGesture(EKeys::Right));
+		UI_COMMAND(StartRecording, "Start", "Start the debugger", EUserInterfaceActionType::Button, FInputChord());
+		UI_COMMAND(StopRecording, "Stop", "Step over the current message", EUserInterfaceActionType::Button, FInputChord());
+		UI_COMMAND(Pause, "Pause", "Stop the debugger", EUserInterfaceActionType::Button, FInputChord());
+		UI_COMMAND(Resume, "Resume", "Stop the debugger", EUserInterfaceActionType::Button, FInputChord());
+		UI_COMMAND(FreeCamera, "FreeCamera", "Enable free camera", EUserInterfaceActionType::ToggleButton, FInputChord());
+		UI_COMMAND(LoadFromVLog, "Load", "Load external vlogs", EUserInterfaceActionType::Button, FInputChord());
+		UI_COMMAND(SaveToVLog, "Save", "Save selected data to vlog file", EUserInterfaceActionType::Button, FInputChord());
+		UI_COMMAND(ToggleGraphs, "ToggleGraphs", "Toggle graphs visualization on/off", EUserInterfaceActionType::ToggleButton, FInputChord());
+		UI_COMMAND(ResetData, "Clear", "Clear all data", EUserInterfaceActionType::Button, FInputChord());
+		UI_COMMAND(HideEmpty, "HideEmpty", "Hide logs without text information", EUserInterfaceActionType::ToggleButton, FInputChord());
+
+		UI_COMMAND(MoveCursorLeft, "MoveCursorLeft", "Move cursor left on timeline", EUserInterfaceActionType::Button, FInputChord(EKeys::Left));
+		UI_COMMAND(MoveCursorRight, "MoveCursorLeft", "Move cursor right on timeline", EUserInterfaceActionType::Button, FInputChord(EKeys::Right));
 	}
 
 public:
diff --git a/Engine/Source/Developer/Profiler/Private/ProfilerCommands.cpp b/Engine/Source/Developer/Profiler/Private/ProfilerCommands.cpp
--- a/Engine/Source/Developer/Profiler/Private/ProfilerCommands.cpp
+++ b/Engine/Source/Developer/Profiler/Private/ProfilerCommands.cpp
@@ -26,40 +26,40 @@ void FProfilerCommands::RegisterCommands()
 		Global and custom commands.	
 	-----------------------------------------------------------------------------*/
 
-	UI_COMMAND( ToggleDataPreview, 	"Data Preview", "Toggles the data preview", EUserInterfaceActionType::ToggleButton, FInputGesture( EModifierKey::Control, EKeys::R ) );
-	UI_COMMAND( ToggleDataCapture, "Data Capture", "Toggles the data capture", EUserInterfaceActionType::ToggleButton, FInputGesture( EModifierKey::Control, EKeys::C ) );
-	UI_COMMAND( ToggleShowDataGraph, "Show Data Graph", "Toggles showing all data graphs", EUserInterfaceActionType::ToggleButton, FInputGesture() );
-	UI_COMMAND( OpenEventGraph, "Open Event Graph", "Opens a new event graph", EUserInterfaceActionType::Button, FInputGesture() );
+	UI_COMMAND( ToggleDataPreview, 	"Data Preview", "Toggles the data preview", EUserInterfaceActionType::ToggleButton, FInputChord( EModifierKey::Control, EKeys::R ) );
+	UI_COMMAND( ToggleDataCapture, "Data Capture", "Toggles the data capture", EUserInterfaceActionType::ToggleButton, FInputChord( EModifierKey::Control, EKeys::C ) );
+	UI_COMMAND( ToggleShowDataGraph, "Show Data Graph", "Toggles showing all data graphs", EUserInterfaceActionType::ToggleButton, FInputChord() );
+	UI_COMMAND( OpenEventGraph, "Open Event Graph", "Opens a new event graph", EUserInterfaceActionType::Button, FInputChord() );
 
 	/*-----------------------------------------------------------------------------
 		Global commands.
 	-----------------------------------------------------------------------------*/
 
-	UI_COMMAND( ProfilerManager_Save, "Save", "Saves all collected data to file or files", EUserInterfaceActionType::Button, FInputGesture( EModifierKey::Control, EKeys::S ) );
-	UI_COMMAND( StatsProfiler, "Statistics", "Enables the Stats Profiler", EUserInterfaceActionType::ToggleButton, FInputGesture( EModifierKey::Control, EKeys::P ) );
+	UI_COMMAND( ProfilerManager_Save, "Save", "Saves all collected data to file or files", EUserInterfaceActionType::Button, FInputChord( EModifierKey::Control, EKeys::S ) );
+	UI_COMMAND( StatsProfiler, "Statistics", "Enables the Stats Profiler", EUserInterfaceActionType::ToggleButton, FInputChord( EModifierKey::Control, EKeys::P ) );
 #if PLATFORM_MAC
-	UI_COMMAND( MemoryProfiler, "Memory", "Enables the Memory Profiler", EUserInterfaceActionType::ToggleButton, FInputGesture( EModifierKey::Command, EKeys::M ) );
-	UI_COMMAND( FPSChart, "FPS Chart", "Shows the FPS Chart", EUserInterfaceActionType::ToggleButton, FInputGesture( EModifierKey::Command, EKeys::H ) );
+	UI_COMMAND( MemoryProfiler, "Memory", "Enables the Memory Profiler", EUserInterfaceActionType::ToggleButton, FInputChord( EModifierKey::Command, EKeys::M ) );
+	UI_COMMAND( FPSChart, "FPS Chart", "Shows the FPS Chart", EUserInterfaceActionType::ToggleButton, FInputChord( EModifierKey::Command, EKeys::H ) );
 #else
-	UI_COMMAND( MemoryProfiler, "Memory", "Enables the Memory Profiler", EUserInterfaceActionType::ToggleButton, FInputGesture( EModifierKey::Control, EKeys::M ) );
-	UI_COMMAND( FPSChart, "FPS Chart", "Shows the FPS Chart", EUserInterfaceActionType::ToggleButton, FInputGesture( EModifierKey::Control, EKeys::H ) );
+	UI_COMMAND( MemoryProfiler, "Memory", "Enables the Memory Profiler", EUserInterfaceActionType::ToggleButton, FInputChord( EModifierKey::Control, EKeys::M ) );
+	UI_COMMAND( FPSChart, "FPS Chart", "Shows the FPS Chart", EUserInterfaceActionType::ToggleButton, FInputChord( EModifierKey::Control, EKeys::H ) );
 #endif
 
-	UI_COMMAND( OpenSettings, "Settings", "Opens the settings for the profiler", EUserInterfaceActionType::Button, FInputGesture( EModifierKey::Control, EKeys::O ) );
+	UI_COMMAND( OpenSettings, "Settings", "Opens the settings for the profiler", EUserInterfaceActionType::Button, FInputChord( EModifierKey::Control, EKeys::O ) );
 
-	UI_COMMAND( ProfilerManager_Load, "Load", "", EUserInterfaceActionType::Button, FInputGesture(/* EModifierKey::Control, EKeys::L */) );
-	UI_COMMAND( ProfilerManager_ToggleLivePreview, "Live preview", "Toggles the real time live preview", EUserInterfaceActionType::ToggleButton, FInputGesture() );
+	UI_COMMAND( ProfilerManager_Load, "Load", "", EUserInterfaceActionType::Button, FInputChord(/* EModifierKey::Control, EKeys::L */) );
+	UI_COMMAND( ProfilerManager_ToggleLivePreview, "Live preview", "Toggles the real time live preview", EUserInterfaceActionType::ToggleButton, FInputChord() );
 
-	UI_COMMAND( DataGraph_ToggleViewMode, "Toggle graph view mode", "Toggles the data graph view mode between time based and index based", EUserInterfaceActionType::Button, FInputGesture() );
-	UI_COMMAND( DataGraph_ToggleMultiMode, "Toggle graph multi mode", "Toggles the data graph multi mode between displaying area line graph and one line graph for each graph", EUserInterfaceActionType::Button, FInputGesture() );
+	UI_COMMAND( DataGraph_ToggleViewMode, "Toggle graph view mode", "Toggles the data graph view mode between time based and index based", EUserInterfaceActionType::Button, FInputChord() );
+	UI_COMMAND( DataGraph_ToggleMultiMode, "Toggle graph multi mode", "Toggles the data graph multi mode between displaying area line graph and one line graph for each graph", EUserInterfaceActionType::Button, FInputChord() );
 
-	UI_COMMAND( DataGraph_ViewMode_SetTimeBased, "Time based", "Sets the data graph view mode to the time based", EUserInterfaceActionType::RadioButton, FInputGesture() );
-	UI_COMMAND( DataGraph_ViewMode_SetIndexBased, "Index based", "Sets the data graph view mode to the index based", EUserInterfaceActionType::RadioButton, FInputGesture() );
+	UI_COMMAND( DataGraph_ViewMode_SetTimeBased, "Time based", "Sets the data graph view mode to the time based", EUserInterfaceActionType::RadioButton, FInputChord() );
+	UI_COMMAND( DataGraph_ViewMode_SetIndexBased, "Index based", "Sets the data graph view mode to the index based", EUserInterfaceActionType::RadioButton, FInputChord() );
 
-	UI_COMMAND( DataGraph_MultiMode_SetCombined, "Combined", "Set the data graph multi mode to the displaying area line graph", EUserInterfaceActionType::RadioButton, FInputGesture() );
-	UI_COMMAND( DataGraph_MultiMode_SetOneLinePerDataSource, "One line per data source", "Set the data graph multi mode to the displaying one line graph for each graph data source", EUserInterfaceActionType::RadioButton, FInputGesture() );
+	UI_COMMAND( DataGraph_MultiMode_SetCombined, "Combined", "Set the data graph multi mode to the displaying area line graph", EUserInterfaceActionType::RadioButton, FInputChord() );
+	UI_COMMAND( DataGraph_MultiMode_SetOneLinePerDataSource, "One line per data source", "Set the data graph multi mode to the displaying one line graph for each graph data source", EUserInterfaceActionType::RadioButton, FInputChord() );
 
-	UI_COMMAND( EventGraph_SelectAllFrames, "Select all frames", "Selects all frames in the data graph and displays them in the event graph", EUserInterfaceActionType::Button, FInputGesture() );
+	UI_COMMAND( EventGraph_SelectAllFrames, "Select all frames", "Selects all frames in the data graph and displays them in the event graph", EUserInterfaceActionType::Button, FInputChord() );
 }
 PRAGMA_ENABLE_OPTIMIZATION
 
diff --git a/Engine/Source/Developer/ProjectLauncher/Private/Models/ProjectLauncherCommands.h b/Engine/Source/Developer/ProjectLauncher/Private/Models/ProjectLauncherCommands.h
--- a/Engine/Source/Developer/ProjectLauncher/Private/Models/ProjectLauncherCommands.h
+++ b/Engine/Source/Developer/ProjectLauncher/Private/Models/ProjectLauncherCommands.h
@@ -30,14 +30,14 @@ public:
 	PRAGMA_DISABLE_OPTIMIZATION
 	virtual void RegisterCommands( ) override
 	{
-		UI_COMMAND(QuickLaunch, "Quick Launch", "Builds, cooks, and launches a build.", EUserInterfaceActionType::ToggleButton, FInputGesture(EModifierKey::Control, EKeys::L));
-		UI_COMMAND(CreateBuild, "Build", "Creates a build.", EUserInterfaceActionType::ToggleButton, FInputGesture(EModifierKey::Control, EKeys::B));
-		UI_COMMAND(DeployBuild, "Deploy Build", "Deploys a pre-made build.", EUserInterfaceActionType::ToggleButton, FInputGesture(EModifierKey::Control, EKeys::D));
-		UI_COMMAND(AdvancedBuild, "Advanced...", "Advanced launcher.", EUserInterfaceActionType::ToggleButton, FInputGesture(EModifierKey::Control, EKeys::A));
-		UI_COMMAND(CloseSettings, "Back", "Close the launch profile settings editor and go back to the launcher.", EUserInterfaceActionType::Button, FInputGesture(EModifierKey::None, EKeys::Escape));
-		UI_COMMAND(RenameProfile, "Rename", "Rename the launch profile.", EUserInterfaceActionType::Button, FInputGesture());// , FInputGesture(EModifierKey::None, EKeys::F2));
-		UI_COMMAND(DuplicateProfile, "Duplicate", "Duplicate the launch profile.", EUserInterfaceActionType::Button, FInputGesture());//, FInputGesture(EModifierKey::Control, EKeys::W));
-		UI_COMMAND(DeleteProfile, "Delete", "Delete the launch profile.", EUserInterfaceActionType::Button, FInputGesture());//, FInputGesture(EModifierKey::Control, EKeys::Delete));
+		UI_COMMAND(QuickLaunch, "Quick Launch", "Builds, cooks, and launches a build.", EUserInterfaceActionType::ToggleButton, FInputChord(EModifierKey::Control, EKeys::L));
+		UI_COMMAND(CreateBuild, "Build", "Creates a build.", EUserInterfaceActionType::ToggleButton, FInputChord(EModifierKey::Control, EKeys::B));
+		UI_COMMAND(DeployBuild, "Deploy Build", "Deploys a pre-made build.", EUserInterfaceActionType::ToggleButton, FInputChord(EModifierKey::Control, EKeys::D));
+		UI_COMMAND(AdvancedBuild, "Advanced...", "Advanced launcher.", EUserInterfaceActionType::ToggleButton, FInputChord(EModifierKey::Control, EKeys::A));
+		UI_COMMAND(CloseSettings, "Back", "Close the launch profile settings editor and go back to the launcher.", EUserInterfaceActionType::Button, FInputChord(EModifierKey::None, EKeys::Escape));
+		UI_COMMAND(RenameProfile, "Rename", "Rename the launch profile.", EUserInterfaceActionType::Button, FInputChord());// , FInputChord(EModifierKey::None, EKeys::F2));
+		UI_COMMAND(DuplicateProfile, "Duplicate", "Duplicate the launch profile.", EUserInterfaceActionType::Button, FInputChord());//, FInputChord(EModifierKey::Control, EKeys::W));
+		UI_COMMAND(DeleteProfile, "Delete", "Delete the launch profile.", EUserInterfaceActionType::Button, FInputChord());//, FInputChord(EModifierKey::Control, EKeys::Delete));
 	}
 	PRAGMA_ENABLE_OPTIMIZATION
 
diff --git a/Engine/Source/Developer/SessionFrontend/Private/Models/SessionConsoleCommands.h b/Engine/Source/Developer/SessionFrontend/Private/Models/SessionConsoleCommands.h
--- a/Engine/Source/Developer/SessionFrontend/Private/Models/SessionConsoleCommands.h
+++ b/Engine/Source/Developer/SessionFrontend/Private/Models/SessionConsoleCommands.h
@@ -27,9 +27,9 @@ public:
 
 	virtual void RegisterCommands() override
 	{
-		UI_COMMAND(Clear, "Clear Log", "Clear the log window", EUserInterfaceActionType::Button, FInputGesture());
-		UI_COMMAND(SessionCopy, "Copy", "Copy the selected log messages to the clipboard", EUserInterfaceActionType::Button, FInputGesture());
-		UI_COMMAND(SessionSave, "Save Log...", "Save the entire log to a file", EUserInterfaceActionType::ToggleButton, FInputGesture());
+		UI_COMMAND(Clear, "Clear Log", "Clear the log window", EUserInterfaceActionType::Button, FInputChord());
+		UI_COMMAND(SessionCopy, "Copy", "Copy the selected log messages to the clipboard", EUserInterfaceActionType::Button, FInputChord());
+		UI_COMMAND(SessionSave, "Save Log...", "Save the entire log to a file", EUserInterfaceActionType::ToggleButton, FInputChord());
 	}
 
 public:
diff --git a/Engine/Source/Editor/BehaviorTreeEditor/Private/BehaviorTreeEditorCommands.cpp b/Engine/Source/Editor/BehaviorTreeEditor/Private/BehaviorTreeEditorCommands.cpp
--- a/Engine/Source/Editor/BehaviorTreeEditor/Private/BehaviorTreeEditorCommands.cpp
+++ b/Engine/Source/Editor/BehaviorTreeEditor/Private/BehaviorTreeEditorCommands.cpp
@@ -12,8 +12,8 @@ FBTCommonCommands::FBTCommonCommands()
 
 void FBTCommonCommands::RegisterCommands()
 {
-	UI_COMMAND(SearchBT, "Search", "Search this Behavior Tree.", EUserInterfaceActionType::Button, FInputGesture(EModifierKey::Control, EKeys::F));
-	UI_COMMAND(NewBlackboard, "New Blackboard", "Create a new Blackboard Data Asset", EUserInterfaceActionType::Button, FInputGesture());
+	UI_COMMAND(SearchBT, "Search", "Search this Behavior Tree.", EUserInterfaceActionType::Button, FInputChord(EModifierKey::Control, EKeys::F));
+	UI_COMMAND(NewBlackboard, "New Blackboard", "Create a new Blackboard Data Asset", EUserInterfaceActionType::Button, FInputChord());
 }
 
 FBTDebuggerCommands::FBTDebuggerCommands() 
@@ -23,18 +23,18 @@ FBTDebuggerCommands::FBTDebuggerCommands()
 
 void FBTDebuggerCommands::RegisterCommands()
 {
-	UI_COMMAND(BackInto, "Back: Into", "Show state from previous step, can go into subtrees", EUserInterfaceActionType::Button, FInputGesture());
-	UI_COMMAND(BackOver, "Back: Over", "Show state from previous step, don't go into subtrees", EUserInterfaceActionType::Button, FInputGesture());
-	UI_COMMAND(ForwardInto, "Forward: Into", "Show state from next step, can go into subtrees", EUserInterfaceActionType::Button, FInputGesture());
-	UI_COMMAND(ForwardOver, "Forward: Over", "Show state from next step, don't go into subtrees", EUserInterfaceActionType::Button, FInputGesture());
-	UI_COMMAND(StepOut, "Step Out", "Show state from next step, leave current subtree", EUserInterfaceActionType::Button, FInputGesture());
+	UI_COMMAND(BackInto, "Back: Into", "Show state from previous step, can go into subtrees", EUserInterfaceActionType::Button, FInputChord());
+	UI_COMMAND(BackOver, "Back: Over", "Show state from previous step, don't go into subtrees", EUserInterfaceActionType::Button, FInputChord());
+	UI_COMMAND(ForwardInto, "Forward: Into", "Show state from next step, can go into subtrees", EUserInterfaceActionType::Button, FInputChord());
+	UI_COMMAND(ForwardOver, "Forward: Over", "Show state from next step, don't go into subtrees", EUserInterfaceActionType::Button, FInputChord());
+	UI_COMMAND(StepOut, "Step Out", "Show state from next step, leave current subtree", EUserInterfaceActionType::Button, FInputChord());
 	
-	UI_COMMAND(PausePlaySession, "Pause", "Pause simulation", EUserInterfaceActionType::Button, FInputGesture());
-	UI_COMMAND(ResumePlaySession, "Resume", "Resume simulation", EUserInterfaceActionType::Button, FInputGesture() );
-	UI_COMMAND(StopPlaySession, "Stop", "Stop simulation", EUserInterfaceActionType::Button, FInputGesture());
+	UI_COMMAND(PausePlaySession, "Pause", "Pause simulation", EUserInterfaceActionType::Button, FInputChord());
+	UI_COMMAND(ResumePlaySession, "Resume", "Resume simulation", EUserInterfaceActionType::Button, FInputChord() );
+	UI_COMMAND(StopPlaySession, "Stop", "Stop simulation", EUserInterfaceActionType::Button, FInputChord());
 
-	UI_COMMAND(CurrentValues, "Current", "View current values", EUserInterfaceActionType::RadioButton, FInputGesture());
-	UI_COMMAND(SavedValues, "Saved", "View saved values", EUserInterfaceActionType::RadioButton, FInputGesture());
+	UI_COMMAND(CurrentValues, "Current", "View current values", EUserInterfaceActionType::RadioButton, FInputChord());
+	UI_COMMAND(SavedValues, "Saved", "View saved values", EUserInterfaceActionType::RadioButton, FInputChord());
 }
 
 FBTBlackboardCommands::FBTBlackboardCommands() 
@@ -44,7 +44,7 @@ FBTBlackboardCommands::FBTBlackboardCommands()
 
 void FBTBlackboardCommands::RegisterCommands()
 {
-	UI_COMMAND(DeleteEntry, "Delete", "Delete this blackboard entry", EUserInterfaceActionType::Button, FInputGesture(EKeys::Platform_Delete));
+	UI_COMMAND(DeleteEntry, "Delete", "Delete this blackboard entry", EUserInterfaceActionType::Button, FInputChord(EKeys::Platform_Delete));
 }
 
 #undef LOCTEXT_NAMESPACE
diff --git a/Engine/Source/Editor/BehaviorTreeEditor/Private/SBehaviorTreeDiff.cpp b/Engine/Source/Editor/BehaviorTreeEditor/Private/SBehaviorTreeDiff.cpp
--- a/Engine/Source/Editor/BehaviorTreeEditor/Private/SBehaviorTreeDiff.cpp
+++ b/Engine/Source/Editor/BehaviorTreeEditor/Private/SBehaviorTreeDiff.cpp
@@ -62,8 +62,8 @@ public:
 	/** Initialize commands */
 	virtual void RegisterCommands() override
 	{
-		UI_COMMAND(Previous, "Prev", "Go to previous difference", EUserInterfaceActionType::Button, FInputGesture(EKeys::F7, EModifierKey::Control));
-		UI_COMMAND(Next, "Next", "Go to next difference", EUserInterfaceActionType::Button, FInputGesture(EKeys::F7));
+		UI_COMMAND(Previous, "Prev", "Go to previous difference", EUserInterfaceActionType::Button, FInputChord(EKeys::F7, EModifierKey::Control));
+		UI_COMMAND(Next, "Next", "Go to next difference", EUserInterfaceActionType::Button, FInputChord(EKeys::F7));
 	}
 
 	/** Go to previous difference */
diff --git a/Engine/Source/Editor/Cascade/Private/CascadeActions.cpp b/Engine/Source/Editor/Cascade/Private/CascadeActions.cpp
--- a/Engine/Source/Editor/Cascade/Private/CascadeActions.cpp
+++ b/Engine/Source/Editor/Cascade/Private/CascadeActions.cpp
@@ -10,71 +10,71 @@
 PRAGMA_DISABLE_OPTIMIZATION
 void FCascadeCommands::RegisterCommands()
 {
-	UI_COMMAND(RestartSimulation, "Restart Sim", "Restart Simulation", EUserInterfaceActionType::Button, FInputGesture());
-	UI_COMMAND(RestartInLevel, "Restart Level", "Restart in Level", EUserInterfaceActionType::Button, FInputGesture(EKeys::SpaceBar));
-	UI_COMMAND(SaveThumbnailImage, "Thumbnail", "Generate Thumbnail", EUserInterfaceActionType::Button, FInputGesture());
-	UI_COMMAND(CascadePlay, "Play/Pause", "Play/Pause Simulation", EUserInterfaceActionType::ToggleButton, FInputGesture());
-	UI_COMMAND(AnimSpeed_100, "100%", "100% Speed", EUserInterfaceActionType::RadioButton, FInputGesture());
-	UI_COMMAND(AnimSpeed_50, "50%", "50% Speed", EUserInterfaceActionType::RadioButton, FInputGesture());
-	UI_COMMAND(AnimSpeed_25, "25%", "25% Speed", EUserInterfaceActionType::RadioButton, FInputGesture());
-	UI_COMMAND(AnimSpeed_10, "10%", "10% Speed", EUserInterfaceActionType::RadioButton, FInputGesture());
-	UI_COMMAND(AnimSpeed_1, "1%", "1% Speed", EUserInterfaceActionType::RadioButton, FInputGesture());
-	UI_COMMAND(RegenerateLowestLODDuplicatingHighest, "Regen LOD", "Regenerate Lowest LOD, Duplicating Highest", EUserInterfaceActionType::Button, FInputGesture());
-	UI_COMMAND(RegenerateLowestLOD, "Regen LOD", "Regenerate Lowest LOD", EUserInterfaceActionType::Button, FInputGesture());
-	UI_COMMAND(JumpToHighestLOD, "Highest LOD", "Select Highest LOD", EUserInterfaceActionType::Button, FInputGesture());
-	UI_COMMAND(JumpToHigherLOD, "Higher LOD", "Select Higher LOD", EUserInterfaceActionType::Button, FInputGesture());
-	UI_COMMAND(AddLODAfterCurrent, "Add LOD", "Add LOD After Current", EUserInterfaceActionType::Button, FInputGesture());
-	UI_COMMAND(AddLODBeforeCurrent, "Add LOD", "Add LOD Before Current", EUserInterfaceActionType::Button, FInputGesture());
-	UI_COMMAND(JumpToLowerLOD, "Lower LOD", "Select Lower LOD", EUserInterfaceActionType::Button, FInputGesture());
-	UI_COMMAND(JumpToLowestLOD, "Lowest LOD", "Select Lowest LOD", EUserInterfaceActionType::Button, FInputGesture());
-	UI_COMMAND(DeleteLOD, "Delete LOD", "Delete Selected LOD", EUserInterfaceActionType::Button, FInputGesture());
+	UI_COMMAND(RestartSimulation, "Restart Sim", "Restart Simulation", EUserInterfaceActionType::Button, FInputChord());
+	UI_COMMAND(RestartInLevel, "Restart Level", "Restart in Level", EUserInterfaceActionType::Button, FInputChord(EKeys::SpaceBar));
+	UI_COMMAND(SaveThumbnailImage, "Thumbnail", "Generate Thumbnail", EUserInterfaceActionType::Button, FInputChord());
+	UI_COMMAND(CascadePlay, "Play/Pause", "Play/Pause Simulation", EUserInterfaceActionType::ToggleButton, FInputChord());
+	UI_COMMAND(AnimSpeed_100, "100%", "100% Speed", EUserInterfaceActionType::RadioButton, FInputChord());
+	UI_COMMAND(AnimSpeed_50, "50%", "50% Speed", EUserInterfaceActionType::RadioButton, FInputChord());
+	UI_COMMAND(AnimSpeed_25, "25%", "25% Speed", EUserInterfaceActionType::RadioButton, FInputChord());
+	UI_COMMAND(AnimSpeed_10, "10%", "10% Speed", EUserInterfaceActionType::RadioButton, FInputChord());
+	UI_COMMAND(AnimSpeed_1, "1%", "1% Speed", EUserInterfaceActionType::RadioButton, FInputChord());
+	UI_COMMAND(RegenerateLowestLODDuplicatingHighest, "Regen LOD", "Regenerate Lowest LOD, Duplicating Highest", EUserInterfaceActionType::Button, FInputChord());
+	UI_COMMAND(RegenerateLowestLOD, "Regen LOD", "Regenerate Lowest LOD", EUserInterfaceActionType::Button, FInputChord());
+	UI_COMMAND(JumpToHighestLOD, "Highest LOD", "Select Highest LOD", EUserInterfaceActionType::Button, FInputChord());
+	UI_COMMAND(JumpToHigherLOD, "Higher LOD", "Select Higher LOD", EUserInterfaceActionType::Button, FInputChord());
+	UI_COMMAND(AddLODAfterCurrent, "Add LOD", "Add LOD After Current", EUserInterfaceActionType::Button, FInputChord());
+	UI_COMMAND(AddLODBeforeCurrent, "Add LOD", "Add LOD Before Current", EUserInterfaceActionType::Button, FInputChord());
+	UI_COMMAND(JumpToLowerLOD, "Lower LOD", "Select Lower LOD", EUserInterfaceActionType::Button, FInputChord());
+	UI_COMMAND(JumpToLowestLOD, "Lowest LOD", "Select Lowest LOD", EUserInterfaceActionType::Button, FInputChord());
+	UI_COMMAND(DeleteLOD, "Delete LOD", "Delete Selected LOD", EUserInterfaceActionType::Button, FInputChord());
 
-	UI_COMMAND(ToggleOriginAxis, "Origin Axis", "Display Origin Axis", EUserInterfaceActionType::ToggleButton, FInputGesture());
-	UI_COMMAND(View_ParticleCounts, "Particle Counts", "Display Particle Counts", EUserInterfaceActionType::ToggleButton, FInputGesture());
-	UI_COMMAND(View_ParticleEventCounts, "Particle Event Counts", "Display Particle Event Counts", EUserInterfaceActionType::ToggleButton, FInputGesture());
-	UI_COMMAND(View_ParticleTimes, "Particle Times", "Display Particle Times", EUserInterfaceActionType::ToggleButton, FInputGesture());
-	UI_COMMAND(View_ParticleMemory, "Particle Memory", "Display Particle Memory", EUserInterfaceActionType::ToggleButton, FInputGesture());
-	UI_COMMAND(ViewMode_Wireframe, "Wireframe", "Select Wireframe Render Mode", EUserInterfaceActionType::RadioButton, FInputGesture());
-	UI_COMMAND(ViewMode_Unlit, "Unlit", "Select Unlit Render Mode", EUserInterfaceActionType::RadioButton, FInputGesture());
-	UI_COMMAND(ViewMode_Lit, "Lit", "Select Lit Render Mode", EUserInterfaceActionType::RadioButton, FInputGesture());
-	UI_COMMAND(ViewMode_ShaderComplexity, "Shader Complexity", "Select Shader Complexity Render Mode", EUserInterfaceActionType::RadioButton, FInputGesture());
-	UI_COMMAND(DetailMode_Low, "Low", "Select Low Detail Mode", EUserInterfaceActionType::RadioButton, FInputGesture());
-	UI_COMMAND(DetailMode_Medium, "Medium", "Select Medium Detail Mode", EUserInterfaceActionType::RadioButton, FInputGesture());
-	UI_COMMAND(DetailMode_High, "High", "Select High Detail Mode", EUserInterfaceActionType::RadioButton, FInputGesture());
-	UI_COMMAND(ToggleGeometry, "Geometry", "Display preview mesh", EUserInterfaceActionType::ToggleButton, FInputGesture());
-	UI_COMMAND(ToggleGeometry_Properties, "Geometry Properties", "Display Preview Mesh Properties", EUserInterfaceActionType::Button, FInputGesture());
-	UI_COMMAND(ToggleLocalVectorFields, "Vector Fields", "Display Local Vector Fields", EUserInterfaceActionType::ToggleButton, FInputGesture());
-	UI_COMMAND(ToggleOrbitMode, "Orbit Mode", "Toggle Orbit Mode", EUserInterfaceActionType::ToggleButton, FInputGesture());
-	UI_COMMAND(ToggleMotion, "Motion", "Toggle Motion", EUserInterfaceActionType::ToggleButton, FInputGesture());
-	UI_COMMAND(SetMotionRadius, "Motion Radius", "Set Motion Radius", EUserInterfaceActionType::Button, FInputGesture());
-	UI_COMMAND(ToggleBounds, "Bounds", "Display Bounds", EUserInterfaceActionType::ToggleButton, FInputGesture());
-	UI_COMMAND(ToggleBounds_SetFixedBounds, "Set Fixed Bounds", "Set Fixed Bounds", EUserInterfaceActionType::Button, FInputGesture());
-	UI_COMMAND(TogglePostProcess, "Post Process", "Toggle Post Process", EUserInterfaceActionType::ToggleButton, FInputGesture());
-	UI_COMMAND(ToggleGrid, "Grid", "Display Grid", EUserInterfaceActionType::ToggleButton, FInputGesture());
-	UI_COMMAND(ToggleLoopSystem, "Loop", "Toggle Looping", EUserInterfaceActionType::ToggleButton, FInputGesture());
-	UI_COMMAND(ToggleRealtime, "Realtime", "Toggles real time rendering in this viewport", EUserInterfaceActionType::ToggleButton, FInputGesture());
-	UI_COMMAND(CascadeBackgroundColor, "Background Color", "Change Background Color", EUserInterfaceActionType::Button, FInputGesture());
-	UI_COMMAND(ToggleWireframeSphere, "Wireframe Sphere", "Display Wireframe Sphere", EUserInterfaceActionType::ToggleButton, FInputGesture());
+	UI_COMMAND(ToggleOriginAxis, "Origin Axis", "Display Origin Axis", EUserInterfaceActionType::ToggleButton, FInputChord());
+	UI_COMMAND(View_ParticleCounts, "Particle Counts", "Display Particle Counts", EUserInterfaceActionType::ToggleButton, FInputChord());
+	UI_COMMAND(View_ParticleEventCounts, "Particle Event Counts", "Display Particle Event Counts", EUserInterfaceActionType::ToggleButton, FInputChord());
+	UI_COMMAND(View_ParticleTimes, "Particle Times", "Display Particle Times", EUserInterfaceActionType::ToggleButton, FInputChord());
+	UI_COMMAND(View_ParticleMemory, "Particle Memory", "Display Particle Memory", EUserInterfaceActionType::ToggleButton, FInputChord());
+	UI_COMMAND(ViewMode_Wireframe, "Wireframe", "Select Wireframe Render Mode", EUserInterfaceActionType::RadioButton, FInputChord());
+	UI_COMMAND(ViewMode_Unlit, "Unlit", "Select Unlit Render Mode", EUserInterfaceActionType::RadioButton, FInputChord());
+	UI_COMMAND(ViewMode_Lit, "Lit", "Select Lit Render Mode", EUserInterfaceActionType::RadioButton, FInputChord());
+	UI_COMMAND(ViewMode_ShaderComplexity, "Shader Complexity", "Select Shader Complexity Render Mode", EUserInterfaceActionType::RadioButton, FInputChord());
+	UI_COMMAND(DetailMode_Low, "Low", "Select Low Detail Mode", EUserInterfaceActionType::RadioButton, FInputChord());
+	UI_COMMAND(DetailMode_Medium, "Medium", "Select Medium Detail Mode", EUserInterfaceActionType::RadioButton, FInputChord());
+	UI_COMMAND(DetailMode_High, "High", "Select High Detail Mode", EUserInterfaceActionType::RadioButton, FInputChord());
+	UI_COMMAND(ToggleGeometry, "Geometry", "Display preview mesh", EUserInterfaceActionType::ToggleButton, FInputChord());
+	UI_COMMAND(ToggleGeometry_Properties, "Geometry Properties", "Display Preview Mesh Properties", EUserInterfaceActionType::Button, FInputChord());
+	UI_COMMAND(ToggleLocalVectorFields, "Vector Fields", "Display Local Vector Fields", EUserInterfaceActionType::ToggleButton, FInputChord());
+	UI_COMMAND(ToggleOrbitMode, "Orbit Mode", "Toggle Orbit Mode", EUserInterfaceActionType::ToggleButton, FInputChord());
+	UI_COMMAND(ToggleMotion, "Motion", "Toggle Motion", EUserInterfaceActionType::ToggleButton, FInputChord());
+	UI_COMMAND(SetMotionRadius, "Motion Radius", "Set Motion Radius", EUserInterfaceActionType::Button, FInputChord());
+	UI_COMMAND(ToggleBounds, "Bounds", "Display Bounds", EUserInterfaceActionType::ToggleButton, FInputChord());
+	UI_COMMAND(ToggleBounds_SetFixedBounds, "Set Fixed Bounds", "Set Fixed Bounds", EUserInterfaceActionType::Button, FInputChord());
+	UI_COMMAND(TogglePostProcess, "Post Process", "Toggle Post Process", EUserInterfaceActionType::ToggleButton, FInputChord());
+	UI_COMMAND(ToggleGrid, "Grid", "Display Grid", EUserInterfaceActionType::ToggleButton, FInputChord());
+	UI_COMMAND(ToggleLoopSystem, "Loop", "Toggle Looping", EUserInterfaceActionType::ToggleButton, FInputChord());
+	UI_COMMAND(ToggleRealtime, "Realtime", "Toggles real time rendering in this viewport", EUserInterfaceActionType::ToggleButton, FInputChord());
+	UI_COMMAND(CascadeBackgroundColor, "Background Color", "Change Background Color", EUserInterfaceActionType::Button, FInputChord());
+	UI_COMMAND(ToggleWireframeSphere, "Wireframe Sphere", "Display Wireframe Sphere", EUserInterfaceActionType::ToggleButton, FInputChord());
 
-	UI_COMMAND(DeleteModule, "Delete Module", "Delete Module", EUserInterfaceActionType::Button, FInputGesture());
-	UI_COMMAND(RefreshModule, "Refresh Module", "Refresh Module", EUserInterfaceActionType::Button, FInputGesture());
-	UI_COMMAND(SyncMaterial, "Sync Material", "Sync Material", EUserInterfaceActionType::Button, FInputGesture());
-	UI_COMMAND(UseMaterial, "Use Material", "Use Material", EUserInterfaceActionType::Button, FInputGesture());
-	UI_COMMAND(DupeFromHigher, "Duplicate From Higher", "Duplicate From Higher LOD", EUserInterfaceActionType::Button, FInputGesture());
-	UI_COMMAND(ShareFromHigher, "Share From Higher", "Share From Higher LOD", EUserInterfaceActionType::Button, FInputGesture());
-	UI_COMMAND(DupeFromHighest, "Duplicate From Highest", "Duplicate From Second Highest LOD", EUserInterfaceActionType::Button, FInputGesture());
-	UI_COMMAND(SetRandomSeed, "Set Random Seed", "Set Random Seed", EUserInterfaceActionType::Button, FInputGesture());
-	UI_COMMAND(ConvertToSeeded, "Convert To Seeded", "Convert To Seeded Module", EUserInterfaceActionType::Button, FInputGesture());
-	UI_COMMAND(RenameEmitter, "Rename Emitter", "Rename Emitter", EUserInterfaceActionType::Button, FInputGesture());
-	UI_COMMAND(DuplicateEmitter, "Duplicate Emitter", "Duplicate Emitter", EUserInterfaceActionType::Button, FInputGesture());
-	UI_COMMAND(DuplicateShareEmitter, "Duplicate and Share Emitter", "Duplicate and Share Emitter", EUserInterfaceActionType::Button, FInputGesture());
-	UI_COMMAND(DeleteEmitter, "Delete Emitter", "Delete Emitter", EUserInterfaceActionType::Button, FInputGesture());
-	UI_COMMAND(ExportEmitter, "Export Emitter", "Export Emitter", EUserInterfaceActionType::Button, FInputGesture());
-	UI_COMMAND(ExportAllEmitters, "Export All", "Export All Emitters", EUserInterfaceActionType::Button, FInputGesture());
-	UI_COMMAND(SelectParticleSystem, "Select Particle System", "Select Particle System", EUserInterfaceActionType::Button, FInputGesture());
-	UI_COMMAND(NewEmitterBefore, "Add New Emitter Before", "Add New Emitter Before Selected", EUserInterfaceActionType::Button, FInputGesture());
-	UI_COMMAND(NewEmitterAfter, "Add New Emitter After", "Add New Emitter After Selected", EUserInterfaceActionType::Button, FInputGesture());
-	UI_COMMAND(RemoveDuplicateModules, "Remove Duplicate Modules", "Remove Duplicate Modules", EUserInterfaceActionType::Button, FInputGesture());
+	UI_COMMAND(DeleteModule, "Delete Module", "Delete Module", EUserInterfaceActionType::Button, FInputChord());
+	UI_COMMAND(RefreshModule, "Refresh Module", "Refresh Module", EUserInterfaceActionType::Button, FInputChord());
+	UI_COMMAND(SyncMaterial, "Sync Material", "Sync Material", EUserInterfaceActionType::Button, FInputChord());
+	UI_COMMAND(UseMaterial, "Use Material", "Use Material", EUserInterfaceActionType::Button, FInputChord());
+	UI_COMMAND(DupeFromHigher, "Duplicate From Higher", "Duplicate From Higher LOD", EUserInterfaceActionType::Button, FInputChord());
+	UI_COMMAND(ShareFromHigher, "Share From Higher", "Share From Higher LOD", EUserInterfaceActionType::Button, FInputChord());
+	UI_COMMAND(DupeFromHighest, "Duplicate From Highest", "Duplicate From Second Highest LOD", EUserInterfaceActionType::Button, FInputChord());
+	UI_COMMAND(SetRandomSeed, "Set Random Seed", "Set Random Seed", EUserInterfaceActionType::Button, FInputChord());
+	UI_COMMAND(ConvertToSeeded, "Convert To Seeded", "Convert To Seeded Module", EUserInterfaceActionType::Button, FInputChord());
+	UI_COMMAND(RenameEmitter, "Rename Emitter", "Rename Emitter", EUserInterfaceActionType::Button, FInputChord());
+	UI_COMMAND(DuplicateEmitter, "Duplicate Emitter", "Duplicate Emitter", EUserInterfaceActionType::Button, FInputChord());
+	UI_COMMAND(DuplicateShareEmitter, "Duplicate and Share Emitter", "Duplicate and Share Emitter", EUserInterfaceActionType::Button, FInputChord());
+	UI_COMMAND(DeleteEmitter, "Delete Emitter", "Delete Emitter", EUserInterfaceActionType::Button, FInputChord());
+	UI_COMMAND(ExportEmitter, "Export Emitter", "Export Emitter", EUserInterfaceActionType::Button, FInputChord());
+	UI_COMMAND(ExportAllEmitters, "Export All", "Export All Emitters", EUserInterfaceActionType::Button, FInputChord());
+	UI_COMMAND(SelectParticleSystem, "Select Particle System", "Select Particle System", EUserInterfaceActionType::Button, FInputChord());
+	UI_COMMAND(NewEmitterBefore, "Add New Emitter Before", "Add New Emitter Before Selected", EUserInterfaceActionType::Button, FInputChord());
+	UI_COMMAND(NewEmitterAfter, "Add New Emitter After", "Add New Emitter After Selected", EUserInterfaceActionType::Button, FInputChord());
+	UI_COMMAND(RemoveDuplicateModules, "Remove Duplicate Modules", "Remove Duplicate Modules", EUserInterfaceActionType::Button, FInputChord());
 }
 
 PRAGMA_ENABLE_OPTIMIZATION
diff --git a/Engine/Source/Editor/ComponentVisualizers/Private/SplineComponentVisualizer.cpp b/Engine/Source/Editor/ComponentVisualizers/Private/SplineComponentVisualizer.cpp
--- a/Engine/Source/Editor/ComponentVisualizers/Private/SplineComponentVisualizer.cpp
+++ b/Engine/Source/Editor/ComponentVisualizers/Private/SplineComponentVisualizer.cpp
@@ -28,14 +28,14 @@ public:
 
 	virtual void RegisterCommands() override
 	{
-		UI_COMMAND(DeleteKey, "Delete Spline Point", "Delete the currently selected spline point.", EUserInterfaceActionType::Button, FInputGesture(EKeys::Delete));
-		UI_COMMAND(DuplicateKey, "Duplicate Spline Point", "Duplicate the currently selected spline point.", EUserInterfaceActionType::Button, FInputGesture());
-		UI_COMMAND(AddKey, "Add Spline Point Here", "Add a new spline point at the cursor location.", EUserInterfaceActionType::Button, FInputGesture());
-		UI_COMMAND(ResetToUnclampedTangent, "Unclamped Tangent", "Reset the tangent for this spline point to its default unclamped value.", EUserInterfaceActionType::Button, FInputGesture());
-		UI_COMMAND(ResetToClampedTangent, "Clamped Tangent", "Reset the tangent for this spline point to its default clamped value.", EUserInterfaceActionType::Button, FInputGesture());
-		UI_COMMAND(SetKeyToCurve, "Curve", "Set spline point to Curve type", EUserInterfaceActionType::RadioButton, FInputGesture());
-		UI_COMMAND(SetKeyToLinear, "Linear", "Set spline point to Linear type", EUserInterfaceActionType::RadioButton, FInputGesture());
-		UI_COMMAND(SetKeyToConstant, "Constant", "Set spline point to Constant type", EUserInterfaceActionType::RadioButton, FInputGesture());
+		UI_COMMAND(DeleteKey, "Delete Spline Point", "Delete the currently selected spline point.", EUserInterfaceActionType::Button, FInputChord(EKeys::Delete));
+		UI_COMMAND(DuplicateKey, "Duplicate Spline Point", "Duplicate the currently selected spline point.", EUserInterfaceActionType::Button, FInputChord());
+		UI_COMMAND(AddKey, "Add Spline Point Here", "Add a new spline point at the cursor location.", EUserInterfaceActionType::Button, FInputChord());
+		UI_COMMAND(ResetToUnclampedTangent, "Unclamped Tangent", "Reset the tangent for this spline point to its default unclamped value.", EUserInterfaceActionType::Button, FInputChord());
+		UI_COMMAND(ResetToClampedTangent, "Clamped Tangent", "Reset the tangent for this spline point to its default clamped value.", EUserInterfaceActionType::Button, FInputChord());
+		UI_COMMAND(SetKeyToCurve, "Curve", "Set spline point to Curve type", EUserInterfaceActionType::RadioButton, FInputChord());
+		UI_COMMAND(SetKeyToLinear, "Linear", "Set spline point to Linear type", EUserInterfaceActionType::RadioButton, FInputChord());
+		UI_COMMAND(SetKeyToConstant, "Constant", "Set spline point to Constant type", EUserInterfaceActionType::RadioButton, FInputChord());
 	}
 
 public:
diff --git a/Engine/Source/Editor/ContentBrowser/Private/ContentBrowserCommands.cpp b/Engine/Source/Editor/ContentBrowser/Private/ContentBrowserCommands.cpp
--- a/Engine/Source/Editor/ContentBrowser/Private/ContentBrowserCommands.cpp
+++ b/Engine/Source/Editor/ContentBrowser/Private/ContentBrowserCommands.cpp
@@ -7,9 +7,9 @@
 
 void FContentBrowserCommands::RegisterCommands()
 {
-	UI_COMMAND( OpenAssetsOrFolders, "Open Assets or Folders", "Opens the selected assets or folders, depending on the selection", EUserInterfaceActionType::Button, FInputGesture(EKeys::Enter) );
-	UI_COMMAND( PreviewAssets, "Preview Assets", "Loads the selected assets and previews them if possible", EUserInterfaceActionType::Button, FInputGesture(EKeys::SpaceBar) );
-	UI_COMMAND( DirectoryUp, "Up", "Up to parent directory", EUserInterfaceActionType::Button, FInputGesture(EModifierKey::Alt, EKeys::Up) );
+	UI_COMMAND( OpenAssetsOrFolders, "Open Assets or Folders", "Opens the selected assets or folders, depending on the selection", EUserInterfaceActionType::Button, FInputChord(EKeys::Enter) );
+	UI_COMMAND( PreviewAssets, "Preview Assets", "Loads the selected assets and previews them if possible", EUserInterfaceActionType::Button, FInputChord(EKeys::SpaceBar) );
+	UI_COMMAND( DirectoryUp, "Up", "Up to parent directory", EUserInterfaceActionType::Button, FInputChord(EModifierKey::Alt, EKeys::Up) );
 }
 
 #undef LOCTEXT_NAMESPACE
\ No newline at end of file
diff --git a/Engine/Source/Editor/DestructibleMeshEditor/Private/DestructibleMeshEditor.cpp b/Engine/Source/Editor/DestructibleMeshEditor/Private/DestructibleMeshEditor.cpp
--- a/Engine/Source/Editor/DestructibleMeshEditor/Private/DestructibleMeshEditor.cpp
+++ b/Engine/Source/Editor/DestructibleMeshEditor/Private/DestructibleMeshEditor.cpp
@@ -26,9 +26,9 @@ static float ExplodeRange = 5.0f;
 
 void FDestructibleMeshEditorCommands::RegisterCommands()
 {
-	UI_COMMAND( Fracture, "Fracture Mesh", "Fractures the mesh's root chunk(s) based upon the Fracture Settings.", EUserInterfaceActionType::Button, FInputGesture() );
-	UI_COMMAND( Refresh, "Refresh", "Refreshes the DestructibleMesh from the StaticMesh it was created from.", EUserInterfaceActionType::Button, FInputGesture() );
-	UI_COMMAND( ImportFBXChunks, "Import FBX Chunks", "Imports a FBX as level 1 chunks.", EUserInterfaceActionType::Button, FInputGesture() );
+	UI_COMMAND( Fracture, "Fracture Mesh", "Fractures the mesh's root chunk(s) based upon the Fracture Settings.", EUserInterfaceActionType::Button, FInputChord() );
+	UI_COMMAND( Refresh, "Refresh", "Refreshes the DestructibleMesh from the StaticMesh it was created from.", EUserInterfaceActionType::Button, FInputChord() );
+	UI_COMMAND( ImportFBXChunks, "Import FBX Chunks", "Imports a FBX as level 1 chunks.", EUserInterfaceActionType::Button, FInputChord() );
 }
 
 
diff --git a/Engine/Source/Editor/DistCurveEditor/Private/CurveEditorActions.cpp b/Engine/Source/Editor/DistCurveEditor/Private/CurveEditorActions.cpp
--- a/Engine/Source/Editor/DistCurveEditor/Private/CurveEditorActions.cpp
+++ b/Engine/Source/Editor/DistCurveEditor/Private/CurveEditorActions.cpp
@@ -7,34 +7,34 @@
 
 void FCurveEditorCommands::RegisterCommands()
 {
-	UI_COMMAND(RemoveCurve, "Remove Curve", "Remove Curve", EUserInterfaceActionType::Button, FInputGesture());
-	UI_COMMAND(RemoveAllCurves, "Remove All Curves", "Remove All Curves", EUserInterfaceActionType::Button, FInputGesture());
-	UI_COMMAND(SetTime, "Set Time", "Set Time", EUserInterfaceActionType::Button, FInputGesture());
-	UI_COMMAND(SetValue, "Set Value", "Set Value", EUserInterfaceActionType::Button, FInputGesture());
-	UI_COMMAND(SetColor, "Set Color", "Set Color", EUserInterfaceActionType::Button, FInputGesture());
-	UI_COMMAND(DeleteKeys, "Delete Key", "Delete Key", EUserInterfaceActionType::Button, FInputGesture());
-	UI_COMMAND(ScaleTimes, "Scale All Times", "Scales the times of all points of all visible tracks", EUserInterfaceActionType::Button, FInputGesture());
-	UI_COMMAND(ScaleValues, "Scale All Values", "Scales the values of all points of all visible tracks", EUserInterfaceActionType::Button, FInputGesture());
-	UI_COMMAND(ScaleSingleCurveTimes, "Scale Curve Times", "Scales the times of all points of this curve & its sub-curves", EUserInterfaceActionType::Button, FInputGesture());
-	UI_COMMAND(ScaleSingleCurveValues, "Scale Curve Values", "Scales the values of all points of this curve & its sub-curves", EUserInterfaceActionType::Button, FInputGesture());
-	UI_COMMAND(ScaleSingleSubCurveValues, "Scale Sub-Curve Values", "Scales the values of all points of this curve", EUserInterfaceActionType::Button, FInputGesture());
-	UI_COMMAND(FitHorizontally, "Horizontal", "Fit Horizontally", EUserInterfaceActionType::Button, FInputGesture());
-	UI_COMMAND(FitVertically, "Vertical", "Fit Vertically", EUserInterfaceActionType::Button, FInputGesture());
-	UI_COMMAND(FitToAll, "All", "Fit to All", EUserInterfaceActionType::Button, FInputGesture());
-	UI_COMMAND(FitToSelected, "Selected", "Fit to Selected", EUserInterfaceActionType::Button, FInputGesture());
-	UI_COMMAND(PanMode, "Pan", "Pan Mode", EUserInterfaceActionType::ToggleButton, FInputGesture());
-	UI_COMMAND(ZoomMode, "Zoom", "Zoom Mode", EUserInterfaceActionType::ToggleButton, FInputGesture());
-	UI_COMMAND(CurveAuto, "Auto", "Curve Auto", EUserInterfaceActionType::ToggleButton, FInputGesture());
-	UI_COMMAND(CurveAutoClamped, "Auto/Clamped", "Curve Auto/Clamped", EUserInterfaceActionType::ToggleButton, FInputGesture());
-	UI_COMMAND(CurveUser, "User", "Curve User", EUserInterfaceActionType::ToggleButton, FInputGesture());
-	UI_COMMAND(CurveBreak, "Break", "Curve Break", EUserInterfaceActionType::ToggleButton, FInputGesture());
-	UI_COMMAND(Linear, "Linear", "Linear", EUserInterfaceActionType::ToggleButton, FInputGesture());
-	UI_COMMAND(Constant, "Constant", "Constant", EUserInterfaceActionType::ToggleButton, FInputGesture());
-	UI_COMMAND(FlattenTangents, "Flatten", "Flatten Tangents", EUserInterfaceActionType::Button, FInputGesture());
-	UI_COMMAND(StraightenTangents, "Straighten", "Straighten Tangents", EUserInterfaceActionType::Button, FInputGesture());
-	UI_COMMAND(ShowAllTangents, "Show All", "ShowAll Tangents", EUserInterfaceActionType::ToggleButton, FInputGesture());
-	UI_COMMAND(CreateTab, "Create", "Create Tab", EUserInterfaceActionType::Button, FInputGesture());
-	UI_COMMAND(DeleteTab, "Delete", "Delete Tab", EUserInterfaceActionType::Button, FInputGesture());
+	UI_COMMAND(RemoveCurve, "Remove Curve", "Remove Curve", EUserInterfaceActionType::Button, FInputChord());
+	UI_COMMAND(RemoveAllCurves, "Remove All Curves", "Remove All Curves", EUserInterfaceActionType::Button, FInputChord());
+	UI_COMMAND(SetTime, "Set Time", "Set Time", EUserInterfaceActionType::Button, FInputChord());
+	UI_COMMAND(SetValue, "Set Value", "Set Value", EUserInterfaceActionType::Button, FInputChord());
+	UI_COMMAND(SetColor, "Set Color", "Set Color", EUserInterfaceActionType::Button, FInputChord());
+	UI_COMMAND(DeleteKeys, "Delete Key", "Delete Key", EUserInterfaceActionType::Button, FInputChord());
+	UI_COMMAND(ScaleTimes, "Scale All Times", "Scales the times of all points of all visible tracks", EUserInterfaceActionType::Button, FInputChord());
+	UI_COMMAND(ScaleValues, "Scale All Values", "Scales the values of all points of all visible tracks", EUserInterfaceActionType::Button, FInputChord());
+	UI_COMMAND(ScaleSingleCurveTimes, "Scale Curve Times", "Scales the times of all points of this curve & its sub-curves", EUserInterfaceActionType::Button, FInputChord());
+	UI_COMMAND(ScaleSingleCurveValues, "Scale Curve Values", "Scales the values of all points of this curve & its sub-curves", EUserInterfaceActionType::Button, FInputChord());
+	UI_COMMAND(ScaleSingleSubCurveValues, "Scale Sub-Curve Values", "Scales the values of all points of this curve", EUserInterfaceActionType::Button, FInputChord());
+	UI_COMMAND(FitHorizontally, "Horizontal", "Fit Horizontally", EUserInterfaceActionType::Button, FInputChord());
+	UI_COMMAND(FitVertically, "Vertical", "Fit Vertically", EUserInterfaceActionType::Button, FInputChord());
+	UI_COMMAND(FitToAll, "All", "Fit to All", EUserInterfaceActionType::Button, FInputChord());
+	UI_COMMAND(FitToSelected, "Selected", "Fit to Selected", EUserInterfaceActionType::Button, FInputChord());
+	UI_COMMAND(PanMode, "Pan", "Pan Mode", EUserInterfaceActionType::ToggleButton, FInputChord());
+	UI_COMMAND(ZoomMode, "Zoom", "Zoom Mode", EUserInterfaceActionType::ToggleButton, FInputChord());
+	UI_COMMAND(CurveAuto, "Auto", "Curve Auto", EUserInterfaceActionType::ToggleButton, FInputChord());
+	UI_COMMAND(CurveAutoClamped, "Auto/Clamped", "Curve Auto/Clamped", EUserInterfaceActionType::ToggleButton, FInputChord());
+	UI_COMMAND(CurveUser, "User", "Curve User", EUserInterfaceActionType::ToggleButton, FInputChord());
+	UI_COMMAND(CurveBreak, "Break", "Curve Break", EUserInterfaceActionType::ToggleButton, FInputChord());
+	UI_COMMAND(Linear, "Linear", "Linear", EUserInterfaceActionType::ToggleButton, FInputChord());
+	UI_COMMAND(Constant, "Constant", "Constant", EUserInterfaceActionType::ToggleButton, FInputChord());
+	UI_COMMAND(FlattenTangents, "Flatten", "Flatten Tangents", EUserInterfaceActionType::Button, FInputChord());
+	UI_COMMAND(StraightenTangents, "Straighten", "Straighten Tangents", EUserInterfaceActionType::Button, FInputChord());
+	UI_COMMAND(ShowAllTangents, "Show All", "ShowAll Tangents", EUserInterfaceActionType::ToggleButton, FInputChord());
+	UI_COMMAND(CreateTab, "Create", "Create Tab", EUserInterfaceActionType::Button, FInputChord());
+	UI_COMMAND(DeleteTab, "Delete", "Delete Tab", EUserInterfaceActionType::Button, FInputChord());
 }
 
 #undef LOCTEXT_NAMESPACE
\ No newline at end of file
diff --git a/Engine/Source/Editor/EditorSettingsViewer/Private/EditorSettingsViewerModule.cpp b/Engine/Source/Editor/EditorSettingsViewer/Private/EditorSettingsViewerModule.cpp
--- a/Engine/Source/Editor/EditorSettingsViewer/Private/EditorSettingsViewerModule.cpp
+++ b/Engine/Source/Editor/EditorSettingsViewer/Private/EditorSettingsViewerModule.cpp
@@ -354,7 +354,7 @@ private:
 	{
 		if( EAppReturnType::Ok == ShowRestartWarning(LOCTEXT("ResetKeyBindings_Title", "Reset Key Bindings")))
 		{
-			FInputBindingManager::Get().RemoveUserDefinedGestures();
+			FInputBindingManager::Get().RemoveUserDefinedChords();
 			GConfig->Flush(false, GEditorKeyBindingsIni);
 			FUnrealEdMisc::Get().RestartEditor(false);
 
@@ -370,7 +370,7 @@ private:
 	// overwrite the imported settings just copied across.
 	bool HandleInputBindingsSave()
 	{
-		FInputBindingManager::Get().RemoveUserDefinedGestures();
+		FInputBindingManager::Get().RemoveUserDefinedChords();
 		GConfig->Flush(false, GEditorKeyBindingsIni);
 		return true;
 	}
diff --git a/Engine/Source/Editor/FoliageEdit/Private/FoliageEditActions.cpp b/Engine/Source/Editor/FoliageEdit/Private/FoliageEditActions.cpp
--- a/Engine/Source/Editor/FoliageEdit/Private/FoliageEditActions.cpp
+++ b/Engine/Source/Editor/FoliageEdit/Private/FoliageEditActions.cpp
@@ -10,20 +10,20 @@
 
 void FFoliageEditCommands::RegisterCommands()
 {
-	UI_COMMAND( SetPaint, "Paint", "Paint", EUserInterfaceActionType::ToggleButton, FInputGesture() );
-	UI_COMMAND( SetReapplySettings, "Reapply", "Reapply settings to instances", EUserInterfaceActionType::ToggleButton, FInputGesture() );
-	UI_COMMAND( SetSelect, "Select", "Select", EUserInterfaceActionType::ToggleButton, FInputGesture() );
-	UI_COMMAND( SetLassoSelect, "Lasso", "Lasso Select", EUserInterfaceActionType::ToggleButton, FInputGesture() );
-	UI_COMMAND( SetPaintBucket, "Fill", "Paint Bucket", EUserInterfaceActionType::ToggleButton, FInputGesture() );
-
-	UI_COMMAND( SetNoSettings, "Hide Details", "Hide details.", EUserInterfaceActionType::ToggleButton, FInputGesture() );
-	UI_COMMAND( SetPaintSettings, "Show Painting settings", "Show painting settings.", EUserInterfaceActionType::ToggleButton, FInputGesture() );
-	UI_COMMAND( SetClusterSettings, "Show Instance settings", "Show settings for placed instances.", EUserInterfaceActionType::ToggleButton, FInputGesture() );
-
-	UI_COMMAND( RemoveFoliageType, "Remove", "Remove all associated foliage instances", EUserInterfaceActionType::Button, FInputGesture() );
-	UI_COMMAND( ShowFoliageTypeInCB, "Show in Content Browser", "Show asset in Content Browser.", EUserInterfaceActionType::Button, FInputGesture() );
-	UI_COMMAND( SelectAllInstances, "Select All Instances", "Select all instances of foliage type.", EUserInterfaceActionType::Button, FInputGesture());
-	UI_COMMAND( DeselectAllInstances, "Deselect All Instances", "Deselect all instances of foliage type.", EUserInterfaceActionType::Button, FInputGesture());
+	UI_COMMAND( SetPaint, "Paint", "Paint", EUserInterfaceActionType::ToggleButton, FInputChord() );
+	UI_COMMAND( SetReapplySettings, "Reapply", "Reapply settings to instances", EUserInterfaceActionType::ToggleButton, FInputChord() );
+	UI_COMMAND( SetSelect, "Select", "Select", EUserInterfaceActionType::ToggleButton, FInputChord() );
+	UI_COMMAND( SetLassoSelect, "Lasso", "Lasso Select", EUserInterfaceActionType::ToggleButton, FInputChord() );
+	UI_COMMAND( SetPaintBucket, "Fill", "Paint Bucket", EUserInterfaceActionType::ToggleButton, FInputChord() );
+
+	UI_COMMAND( SetNoSettings, "Hide Details", "Hide details.", EUserInterfaceActionType::ToggleButton, FInputChord() );
+	UI_COMMAND( SetPaintSettings, "Show Painting settings", "Show painting settings.", EUserInterfaceActionType::ToggleButton, FInputChord() );
+	UI_COMMAND( SetClusterSettings, "Show Instance settings", "Show settings for placed instances.", EUserInterfaceActionType::ToggleButton, FInputChord() );
+
+	UI_COMMAND( RemoveFoliageType, "Remove", "Remove all associated foliage instances", EUserInterfaceActionType::Button, FInputChord() );
+	UI_COMMAND( ShowFoliageTypeInCB, "Show in Content Browser", "Show asset in Content Browser.", EUserInterfaceActionType::Button, FInputChord() );
+	UI_COMMAND( SelectAllInstances, "Select All Instances", "Select all instances of foliage type.", EUserInterfaceActionType::Button, FInputChord());
+	UI_COMMAND( DeselectAllInstances, "Deselect All Instances", "Deselect all instances of foliage type.", EUserInterfaceActionType::Button, FInputChord());
 }
 
 #undef LOCTEXT_NAMESPACE
\ No newline at end of file
diff --git a/Engine/Source/Editor/FontEditor/Private/FontEditor.cpp b/Engine/Source/Editor/FontEditor/Private/FontEditor.cpp
--- a/Engine/Source/Editor/FontEditor/Private/FontEditor.cpp
+++ b/Engine/Source/Editor/FontEditor/Private/FontEditor.cpp
@@ -65,13 +65,13 @@ public:
 
 void FFontEditorCommands::RegisterCommands()
 {
-	UI_COMMAND(Update, "Update", "Imports a texture to replace the currently selected page.", EUserInterfaceActionType::Button, FInputGesture());
-	UI_COMMAND(UpdateAll, "Update All", "Imports a set of textures to replace all pages.", EUserInterfaceActionType::Button, FInputGesture());
-	UI_COMMAND(ExportPage, "Export", "Exports the currently selected page.", EUserInterfaceActionType::Button, FInputGesture());
-	UI_COMMAND(ExportAllPages, "Export All", "Exports all pages.", EUserInterfaceActionType::Button, FInputGesture());
+	UI_COMMAND(Update, "Update", "Imports a texture to replace the currently selected page.", EUserInterfaceActionType::Button, FInputChord());
+	UI_COMMAND(UpdateAll, "Update All", "Imports a set of textures to replace all pages.", EUserInterfaceActionType::Button, FInputChord());
+	UI_COMMAND(ExportPage, "Export", "Exports the currently selected page.", EUserInterfaceActionType::Button, FInputChord());
+	UI_COMMAND(ExportAllPages, "Export All", "Exports all pages.", EUserInterfaceActionType::Button, FInputChord());
 
-	UI_COMMAND(FontBackgroundColor, "Background", "Changes the background color of the previewer.", EUserInterfaceActionType::Button, FInputGesture());
-	UI_COMMAND(FontForegroundColor, "Foreground", "Changes the foreground color of the previewer.", EUserInterfaceActionType::Button, FInputGesture());
+	UI_COMMAND(FontBackgroundColor, "Background", "Changes the background color of the previewer.", EUserInterfaceActionType::Button, FInputChord());
+	UI_COMMAND(FontForegroundColor, "Foreground", "Changes the foreground color of the previewer.", EUserInterfaceActionType::Button, FInputChord());
 }
 
 void FFontEditor::RegisterTabSpawners(const TSharedRef<class FTabManager>& TabManager)
diff --git a/Engine/Source/Editor/GraphEditor/Private/GraphEditorActions.cpp b/Engine/Source/Editor/GraphEditor/Private/GraphEditorActions.cpp
--- a/Engine/Source/Editor/GraphEditor/Private/GraphEditorActions.cpp
+++ b/Engine/Source/Editor/GraphEditor/Private/GraphEditorActions.cpp
@@ -7,75 +7,75 @@
 
 void FGraphEditorCommandsImpl::RegisterCommands()
 {
-	UI_COMMAND( ReconstructNodes, "Refresh Nodes", "Refreshes nodes", EUserInterfaceActionType::Button, FInputGesture() )
-	UI_COMMAND( BreakNodeLinks, "Break Link(s)", "Breaks links", EUserInterfaceActionType::Button, FInputGesture() )
+	UI_COMMAND( ReconstructNodes, "Refresh Nodes", "Refreshes nodes", EUserInterfaceActionType::Button, FInputChord() )
+	UI_COMMAND( BreakNodeLinks, "Break Link(s)", "Breaks links", EUserInterfaceActionType::Button, FInputChord() )
 	
-	UI_COMMAND( AddExecutionPin, "Add execution pin", "Adds another execution output pin to an execution sequence or switch node", EUserInterfaceActionType::Button, FInputGesture() )
-	UI_COMMAND( RemoveExecutionPin, "Remove execution pin", "Removes an execution output pin from an execution sequence or switch node", EUserInterfaceActionType::Button, FInputGesture() )
+	UI_COMMAND( AddExecutionPin, "Add execution pin", "Adds another execution output pin to an execution sequence or switch node", EUserInterfaceActionType::Button, FInputChord() )
+	UI_COMMAND( RemoveExecutionPin, "Remove execution pin", "Removes an execution output pin from an execution sequence or switch node", EUserInterfaceActionType::Button, FInputChord() )
 
-	UI_COMMAND( RemoveThisStructVarPin, "Remove this struct variable pin", "Removes the selected input pin", EUserInterfaceActionType::Button, FInputGesture() )
-	UI_COMMAND( RemoveOtherStructVarPins, "Remove all other pins", "Removes all variable input pins, except for the selected one", EUserInterfaceActionType::Button, FInputGesture() )
+	UI_COMMAND( RemoveThisStructVarPin, "Remove this struct variable pin", "Removes the selected input pin", EUserInterfaceActionType::Button, FInputChord() )
+	UI_COMMAND( RemoveOtherStructVarPins, "Remove all other pins", "Removes all variable input pins, except for the selected one", EUserInterfaceActionType::Button, FInputChord() )
 
-	UI_COMMAND( RestoreAllStructVarPins, "Restore all structure pins", "Restore all structure pins", EUserInterfaceActionType::Button, FInputGesture() )
+	UI_COMMAND( RestoreAllStructVarPins, "Restore all structure pins", "Restore all structure pins", EUserInterfaceActionType::Button, FInputChord() )
 
-	UI_COMMAND( AddOptionPin, "Add Option Pin", "Adds another option input pin to the node", EUserInterfaceActionType::Button, FInputGesture() )
-	UI_COMMAND( RemoveOptionPin, "Remove Option Pin", "Removes the last option input pin from the node", EUserInterfaceActionType::Button, FInputGesture() )
+	UI_COMMAND( AddOptionPin, "Add Option Pin", "Adds another option input pin to the node", EUserInterfaceActionType::Button, FInputChord() )
+	UI_COMMAND( RemoveOptionPin, "Remove Option Pin", "Removes the last option input pin from the node", EUserInterfaceActionType::Button, FInputChord() )
 
-	UI_COMMAND( ChangePinType, "Change Pin Type", "Changes the type of this pin (boolean, int, etc.)", EUserInterfaceActionType::Button, FInputGesture() )
+	UI_COMMAND( ChangePinType, "Change Pin Type", "Changes the type of this pin (boolean, int, etc.)", EUserInterfaceActionType::Button, FInputChord() )
 
-	UI_COMMAND( ShowAllPins, "Show All Pins", "Shows all pins", EUserInterfaceActionType::RadioButton, FInputGesture() )
-	UI_COMMAND( HideNoConnectionPins, "Hide Unconnected Pins", "Hides all pins with no connections", EUserInterfaceActionType::RadioButton, FInputGesture() )
-	UI_COMMAND( HideNoConnectionNoDefaultPins, "Hide Unused Pins", "Hides all pins with no connections and no default value", EUserInterfaceActionType::RadioButton, FInputGesture() )
+	UI_COMMAND( ShowAllPins, "Show All Pins", "Shows all pins", EUserInterfaceActionType::RadioButton, FInputChord() )
+	UI_COMMAND( HideNoConnectionPins, "Hide Unconnected Pins", "Hides all pins with no connections", EUserInterfaceActionType::RadioButton, FInputChord() )
+	UI_COMMAND( HideNoConnectionNoDefaultPins, "Hide Unused Pins", "Hides all pins with no connections and no default value", EUserInterfaceActionType::RadioButton, FInputChord() )
 
-	UI_COMMAND( AddParentNode, "Add call to parent function", "Adds a node that calls this function's parent", EUserInterfaceActionType::Button, FInputGesture() )
+	UI_COMMAND( AddParentNode, "Add call to parent function", "Adds a node that calls this function's parent", EUserInterfaceActionType::Button, FInputChord() )
 
-	UI_COMMAND( ToggleBreakpoint, "Toggle breakpoint", "Adds or removes a breakpoint on each selected node", EUserInterfaceActionType::Button, FInputGesture(EKeys::F9) )
-	UI_COMMAND( AddBreakpoint, "Add breakpoint", "Adds a breakpoint to each selected node", EUserInterfaceActionType::Button, FInputGesture() )
-	UI_COMMAND( RemoveBreakpoint, "Remove breakpoint", "Removes any breakpoints on each selected node", EUserInterfaceActionType::Button, FInputGesture() )
-	UI_COMMAND( EnableBreakpoint, "Enable breakpoint", "Enables any breakpoints on each selected node", EUserInterfaceActionType::Button, FInputGesture() )
-	UI_COMMAND( DisableBreakpoint, "Disable breakpoint", "Disables any breakpoints on each selected node", EUserInterfaceActionType::Button, FInputGesture() )
+	UI_COMMAND( ToggleBreakpoint, "Toggle breakpoint", "Adds or removes a breakpoint on each selected node", EUserInterfaceActionType::Button, FInputChord(EKeys::F9) )
+	UI_COMMAND( AddBreakpoint, "Add breakpoint", "Adds a breakpoint to each selected node", EUserInterfaceActionType::Button, FInputChord() )
+	UI_COMMAND( RemoveBreakpoint, "Remove breakpoint", "Removes any breakpoints on each selected node", EUserInterfaceActionType::Button, FInputChord() )
+	UI_COMMAND( EnableBreakpoint, "Enable breakpoint", "Enables any breakpoints on each selected node", EUserInterfaceActionType::Button, FInputChord() )
+	UI_COMMAND( DisableBreakpoint, "Disable breakpoint", "Disables any breakpoints on each selected node", EUserInterfaceActionType::Button, FInputChord() )
 
-	UI_COMMAND( CollapseNodes, "Collapse Nodes", "Collapses selected nodes into a single node", EUserInterfaceActionType::Button, FInputGesture() )
-	UI_COMMAND( PromoteSelectionToFunction, "Promote to Function", "Promotes selected collapsed graphs to functions.", EUserInterfaceActionType::Button, FInputGesture() )
-	UI_COMMAND( PromoteSelectionToMacro, "Promote to Macro", "Promotes selected collapsed graphs to macros.", EUserInterfaceActionType::Button, FInputGesture() )
-	UI_COMMAND( ExpandNodes, "Expand Node", "Expands composite nodes", EUserInterfaceActionType::Button, FInputGesture() )
-	UI_COMMAND( CollapseSelectionToFunction, "Collapse to Function", "Collapses selected nodes into a single function node.", EUserInterfaceActionType::Button, FInputGesture() )
-	UI_COMMAND( CollapseSelectionToMacro, "Collapse to Macro", "Collapses selected nodes into a single macro node.", EUserInterfaceActionType::Button, FInputGesture() )
+	UI_COMMAND( CollapseNodes, "Collapse Nodes", "Collapses selected nodes into a single node", EUserInterfaceActionType::Button, FInputChord() )
+	UI_COMMAND( PromoteSelectionToFunction, "Promote to Function", "Promotes selected collapsed graphs to functions.", EUserInterfaceActionType::Button, FInputChord() )
+	UI_COMMAND( PromoteSelectionToMacro, "Promote to Macro", "Promotes selected collapsed graphs to macros.", EUserInterfaceActionType::Button, FInputChord() )
+	UI_COMMAND( ExpandNodes, "Expand Node", "Expands composite nodes", EUserInterfaceActionType::Button, FInputChord() )
+	UI_COMMAND( CollapseSelectionToFunction, "Collapse to Function", "Collapses selected nodes into a single function node.", EUserInterfaceActionType::Button, FInputChord() )
+	UI_COMMAND( CollapseSelectionToMacro, "Collapse to Macro", "Collapses selected nodes into a single macro node.", EUserInterfaceActionType::Button, FInputChord() )
 
-	UI_COMMAND( SelectReferenceInLevel, "Find Actor in Level", "Select the actor referenced by this node in the level", EUserInterfaceActionType::Button, FInputGesture() )
-	UI_COMMAND( AssignReferencedActor, "Assign selected Actor", "Assign the selected actor to be this node's referenced object", EUserInterfaceActionType::Button, FInputGesture() )
-	UI_COMMAND( FindVariableReferences, "Find Variable References", "Find references of this variable", EUserInterfaceActionType::Button, FInputGesture() )
+	UI_COMMAND( SelectReferenceInLevel, "Find Actor in Level", "Select the actor referenced by this node in the level", EUserInterfaceActionType::Button, FInputChord() )
+	UI_COMMAND( AssignReferencedActor, "Assign selected Actor", "Assign the selected actor to be this node's referenced object", EUserInterfaceActionType::Button, FInputChord() )
+	UI_COMMAND( FindVariableReferences, "Find Variable References", "Find references of this variable", EUserInterfaceActionType::Button, FInputChord() )
 	
-	UI_COMMAND( GotoNativeFunctionDefinition, "Goto Code Definition", "Goto the native code definition of this function", EUserInterfaceActionType::Button, FInputGesture() )
-	UI_COMMAND( GotoNativeVariableDefinition, "Goto Code Definition", "Goto the native code definition of this variable", EUserInterfaceActionType::Button, FInputGesture() )
+	UI_COMMAND( GotoNativeFunctionDefinition, "Goto Code Definition", "Goto the native code definition of this function", EUserInterfaceActionType::Button, FInputChord() )
+	UI_COMMAND( GotoNativeVariableDefinition, "Goto Code Definition", "Goto the native code definition of this variable", EUserInterfaceActionType::Button, FInputChord() )
 
-	UI_COMMAND( GoToDefinition, "Goto Definition", "Jumps to the graph this node is defined in if available.", EUserInterfaceActionType::Button, FInputGesture() )
+	UI_COMMAND( GoToDefinition, "Goto Definition", "Jumps to the graph this node is defined in if available.", EUserInterfaceActionType::Button, FInputChord() )
 
-	UI_COMMAND( BreakPinLinks, "Break Link(s)", "Breaks pin links", EUserInterfaceActionType::Button, FInputGesture() )
-	UI_COMMAND( PromoteToVariable, "Promote to Variable", "Promotes something to a variable", EUserInterfaceActionType::Button, FInputGesture() )
-	UI_COMMAND( SplitStructPin, "Split Struct Pin", "Breaks a struct pin in to a separate pin per element", EUserInterfaceActionType::Button, FInputGesture() )
-	UI_COMMAND( RecombineStructPin, "Recombine Struct Pin", "Takes struct pins that have been broken in to composite elements and combines them back to a single struct pin", EUserInterfaceActionType::Button, FInputGesture() )
-	UI_COMMAND( StartWatchingPin, "Watch this value", "Adds this pin or variable to the watch list", EUserInterfaceActionType::Button, FInputGesture() )
-	UI_COMMAND( StopWatchingPin, "Stop watching this value", "Removes this pin or variable from the watch list ", EUserInterfaceActionType::Button, FInputGesture() )
+	UI_COMMAND( BreakPinLinks, "Break Link(s)", "Breaks pin links", EUserInterfaceActionType::Button, FInputChord() )
+	UI_COMMAND( PromoteToVariable, "Promote to Variable", "Promotes something to a variable", EUserInterfaceActionType::Button, FInputChord() )
+	UI_COMMAND( SplitStructPin, "Split Struct Pin", "Breaks a struct pin in to a separate pin per element", EUserInterfaceActionType::Button, FInputChord() )
+	UI_COMMAND( RecombineStructPin, "Recombine Struct Pin", "Takes struct pins that have been broken in to composite elements and combines them back to a single struct pin", EUserInterfaceActionType::Button, FInputChord() )
+	UI_COMMAND( StartWatchingPin, "Watch this value", "Adds this pin or variable to the watch list", EUserInterfaceActionType::Button, FInputChord() )
+	UI_COMMAND( StopWatchingPin, "Stop watching this value", "Removes this pin or variable from the watch list ", EUserInterfaceActionType::Button, FInputChord() )
 
-	UI_COMMAND( SelectBone, "Select Bone", "Assign or change the bone for SkeletalControls", EUserInterfaceActionType::Button, FInputGesture() )
-	UI_COMMAND( AddBlendListPin, "Add Blend Pin", "Add Blend Pin to BlendList", EUserInterfaceActionType::Button, FInputGesture() )
-	UI_COMMAND( RemoveBlendListPin, "Remove Blend Pin", "Remove Blend Pin", EUserInterfaceActionType::Button, FInputGesture() )
+	UI_COMMAND( SelectBone, "Select Bone", "Assign or change the bone for SkeletalControls", EUserInterfaceActionType::Button, FInputChord() )
+	UI_COMMAND( AddBlendListPin, "Add Blend Pin", "Add Blend Pin to BlendList", EUserInterfaceActionType::Button, FInputChord() )
+	UI_COMMAND( RemoveBlendListPin, "Remove Blend Pin", "Remove Blend Pin", EUserInterfaceActionType::Button, FInputChord() )
 
-	UI_COMMAND( ConvertToSeqEvaluator, "Convert To Single Frame Animation", "Convert to one frame animation that requires position", EUserInterfaceActionType::Button, FInputGesture() )
-	UI_COMMAND( ConvertToSeqPlayer, "Convert to Sequence Player", "Convert back to sequence player without manual position set up", EUserInterfaceActionType::Button, FInputGesture() )
+	UI_COMMAND( ConvertToSeqEvaluator, "Convert To Single Frame Animation", "Convert to one frame animation that requires position", EUserInterfaceActionType::Button, FInputChord() )
+	UI_COMMAND( ConvertToSeqPlayer, "Convert to Sequence Player", "Convert back to sequence player without manual position set up", EUserInterfaceActionType::Button, FInputChord() )
 
-	UI_COMMAND( ConvertToBSEvaluator, "Convert To Single Frame BlendSpace", "Convert to one frame BlendSpace that requires position", EUserInterfaceActionType::Button, FInputGesture() )
-	UI_COMMAND( ConvertToBSPlayer, "Convert to BlendSpace Player", "Convert back to BlendSpace player without manual position set up", EUserInterfaceActionType::Button, FInputGesture() )
+	UI_COMMAND( ConvertToBSEvaluator, "Convert To Single Frame BlendSpace", "Convert to one frame BlendSpace that requires position", EUserInterfaceActionType::Button, FInputChord() )
+	UI_COMMAND( ConvertToBSPlayer, "Convert to BlendSpace Player", "Convert back to BlendSpace player without manual position set up", EUserInterfaceActionType::Button, FInputChord() )
 
-	UI_COMMAND( OpenRelatedAsset, "Open Asset", "Opens the asset related to this node", EUserInterfaceActionType::Button, FInputGesture() )
+	UI_COMMAND( OpenRelatedAsset, "Open Asset", "Opens the asset related to this node", EUserInterfaceActionType::Button, FInputChord() )
 
-	UI_COMMAND( CreateComment, "Create Comment", "Create a comment box", EUserInterfaceActionType::Button, FInputGesture(EKeys::C))
+	UI_COMMAND( CreateComment, "Create Comment", "Create a comment box", EUserInterfaceActionType::Button, FInputChord(EKeys::C))
 
-	UI_COMMAND( FindInstancesOfCustomEvent, "Find Instances of Event", "Find the instances of this custom event", EUserInterfaceActionType::Button, FInputGesture() )
+	UI_COMMAND( FindInstancesOfCustomEvent, "Find Instances of Event", "Find the instances of this custom event", EUserInterfaceActionType::Button, FInputChord() )
 
-	UI_COMMAND( ZoomIn, "Zoom In", "Zoom in on the graph editor", EUserInterfaceActionType::Button, FInputGesture(EKeys::Add))
-	UI_COMMAND( ZoomOut, "Zoom Out", "Zoom out from the graph editor", EUserInterfaceActionType::Button, FInputGesture(EKeys::Subtract))
+	UI_COMMAND( ZoomIn, "Zoom In", "Zoom in on the graph editor", EUserInterfaceActionType::Button, FInputChord(EKeys::Add))
+	UI_COMMAND( ZoomOut, "Zoom Out", "Zoom out from the graph editor", EUserInterfaceActionType::Button, FInputChord(EKeys::Subtract))
 }
 
 
diff --git a/Engine/Source/Editor/GraphEditor/Private/SGraphPanel.cpp b/Engine/Source/Editor/GraphEditor/Private/SGraphPanel.cpp
--- a/Engine/Source/Editor/GraphEditor/Private/SGraphPanel.cpp
+++ b/Engine/Source/Editor/GraphEditor/Private/SGraphPanel.cpp
@@ -10,7 +10,7 @@
 
 #include "GraphEditorActions.h"
 #include "UICommandInfo.h"
-#include "InputGesture.h"
+#include "InputChord.h"
 
 #include "ConnectionDrawingPolicy.h"
 #include "BlueprintConnectionDrawingPolicy.h"
@@ -556,12 +556,12 @@ FReply SGraphPanel::OnKeyDown( const FGeometry& MyGeometry, const FKeyEvent& InK
 			UpdateSelectedNodesPositions(FVector2D(-GetSnapGridSize(),0.0f));
 			return FReply::Handled();
 		}
-		if(InKeyEvent.GetKey() == FGraphEditorCommands::Get().ZoomOut->GetActiveGesture()->Key)
+		if(InKeyEvent.GetKey() == FGraphEditorCommands::Get().ZoomOut->GetActiveChord()->Key)
 		{
 			ChangeZoomLevel(-1, CachedAllottedGeometryScaledSize / 2.f, InKeyEvent.IsControlDown());
 			return FReply::Handled();
 		}
-		if(InKeyEvent.GetKey() == FGraphEditorCommands::Get().ZoomIn->GetActiveGesture()->Key)
+		if(InKeyEvent.GetKey() == FGraphEditorCommands::Get().ZoomIn->GetActiveChord()->Key)
 		{
 			ChangeZoomLevel(+1, CachedAllottedGeometryScaledSize / 2.f, InKeyEvent.IsControlDown());
 			return FReply::Handled();
