Fixes for warnings found by PVS-Studio, these have been modified from PR#996

diff --git a/Engine/Source/Developer/AudioFormatADPCM/Private/AudioFormatADPCM.cpp b/Engine/Source/Developer/AudioFormatADPCM/Private/AudioFormatADPCM.cpp
--- a/Engine/Source/Developer/AudioFormatADPCM/Private/AudioFormatADPCM.cpp
+++ b/Engine/Source/Developer/AudioFormatADPCM/Private/AudioFormatADPCM.cpp
@@ -160,20 +160,6 @@ namespace
 
 		return OutputValue;
 	}
-
-	template <typename T>
-	inline T WriteToArray(const T* ElementArray, int32& ReadIndex, int32 NumElements, int32 IndexStride = 1)
-	{
-		T OutputValue = 0;
-
-		if (ReadIndex >= 0 && ReadIndex < NumElements)
-		{
-			OutputValue = ElementArray[ReadIndex];
-			ReadIndex += IndexStride;
-		}
-
-		return OutputValue;
-	}
 } // end namespace
 
 namespace LPCM
diff --git a/Engine/Source/Editor/Cascade/Private/Cascade.cpp b/Engine/Source/Editor/Cascade/Private/Cascade.cpp
--- a/Engine/Source/Editor/Cascade/Private/Cascade.cpp
+++ b/Engine/Source/Editor/Cascade/Private/Cascade.cpp
@@ -1072,24 +1072,22 @@ void FCascade::OnNewEmitter()
 	NewEmitter->UpdateModuleLists();
 
 	NewEmitter->PostEditChange();
-	if (NewEmitter)
+
+	NewEmitter->SetFlags(RF_Transactional);
+	for (int32 LODIndex = 0; LODIndex < NewEmitter->LODLevels.Num(); LODIndex++)
 	{
-		NewEmitter->SetFlags(RF_Transactional);
-		for (int32 LODIndex = 0; LODIndex < NewEmitter->LODLevels.Num(); LODIndex++)
+		UParticleLODLevel* NewEmitterLODLevel = NewEmitter->GetLODLevel(LODIndex);
+		if (NewEmitterLODLevel)
 		{
-			UParticleLODLevel* NewEmitterLODLevel = NewEmitter->GetLODLevel(LODIndex);
-			if (NewEmitterLODLevel)
+			NewEmitterLODLevel->SetFlags(RF_Transactional);
+			check(NewEmitterLODLevel->RequiredModule);
+			NewEmitterLODLevel->RequiredModule->SetTransactionFlag();
+			check(NewEmitterLODLevel->SpawnModule);
+			NewEmitterLODLevel->SpawnModule->SetTransactionFlag();
+			for (int32 jj = 0; jj < NewEmitterLODLevel->Modules.Num(); jj++)
 			{
-				NewEmitterLODLevel->SetFlags(RF_Transactional);
-				check(NewEmitterLODLevel->RequiredModule);
-				NewEmitterLODLevel->RequiredModule->SetTransactionFlag();
-				check(NewEmitterLODLevel->SpawnModule);
-				NewEmitterLODLevel->SpawnModule->SetTransactionFlag();
-				for (int32 jj = 0; jj < NewEmitterLODLevel->Modules.Num(); jj++)
-				{
-					UParticleModule* pkModule = NewEmitterLODLevel->Modules[jj];
-					pkModule->SetTransactionFlag();
-				}
+				UParticleModule* pkModule = NewEmitterLODLevel->Modules[jj];
+				pkModule->SetTransactionFlag();
 			}
 		}
 	}
@@ -4392,52 +4390,43 @@ void FCascade::OnDeleteModule(bool bConfirm)
 	}
 	CurveEditor->RefreshViewport();
 
-	if (SelectedEmitter)
+	bool bNeedsListUpdated = false;
+
+	for (int32 LODIndex = 0; LODIndex < SelectedEmitter->LODLevels.Num(); LODIndex++)
 	{
-		bool bNeedsListUpdated = false;
+		UParticleLODLevel* LODLevel	= SelectedEmitter->GetLODLevel(LODIndex);
 
-		for (int32 LODIndex = 0; LODIndex < SelectedEmitter->LODLevels.Num(); LODIndex++)
-		{
-			UParticleLODLevel* LODLevel	= SelectedEmitter->GetLODLevel(LODIndex);
+		// See if it is in this LODs level...
+		UParticleModule* CheckModule;
 
-			// See if it is in this LODs level...
-			UParticleModule* CheckModule;
+		if (DeleteModuleIndex >= 0)
+		{
+			CheckModule = LODLevel->Modules[DeleteModuleIndex];
+		}
+		else
+		{
+			CheckModule = LODLevel->TypeDataModule;
+		}
 
-			if (DeleteModuleIndex >= 0)
-			{
-				CheckModule = LODLevel->Modules[DeleteModuleIndex];
-			}
-			else
+		if (CheckModule)
+		{
+			if (CheckModule->IsA(UParticleModuleTypeDataBase::StaticClass()))
 			{
-				CheckModule = LODLevel->TypeDataModule;
+				check(LODLevel->TypeDataModule == CheckModule);
+				LODLevel->TypeDataModule = NULL;
 			}
-
-			if (CheckModule)
+			else if (CheckModule->IsA(UParticleModuleEventGenerator::StaticClass()))
 			{
-				if (CheckModule->IsA(UParticleModuleTypeDataBase::StaticClass()))
-				{
-					check(LODLevel->TypeDataModule == CheckModule);
-					LODLevel->TypeDataModule = NULL;
-				}
-				else
-					if (CheckModule->IsA(UParticleModuleEventGenerator::StaticClass()))
-					{
-						LODLevel->EventGenerator = NULL;
-					}
-					LODLevel->Modules.Remove(CheckModule);
-					bNeedsListUpdated = true;
+				LODLevel->EventGenerator = NULL;
 			}
-		}
-
-		if (bNeedsListUpdated)
-		{
-			SelectedEmitter->UpdateModuleLists();
+			LODLevel->Modules.Remove(CheckModule);
+			bNeedsListUpdated = true;
 		}
 	}
-	else
+
+	if (bNeedsListUpdated)
 	{
-		// Assume that it's in the module dump...
-		DraggedModuleList.Remove(SelectedModule);
+		SelectedEmitter->UpdateModuleLists();
 	}
 
 	ParticleSystem->PostEditChange();
diff --git a/Engine/Source/Editor/ContentBrowser/Private/SAssetView.cpp b/Engine/Source/Editor/ContentBrowser/Private/SAssetView.cpp
--- a/Engine/Source/Editor/ContentBrowser/Private/SAssetView.cpp
+++ b/Engine/Source/Editor/ContentBrowser/Private/SAssetView.cpp
@@ -3234,15 +3234,8 @@ FReply SAssetView::OnDraggingAssetItem( const FGeometry& MyGeometry, const FPoin
 					continue;
 				}
 
-				if (AssetData.AssetClass == UClass::StaticClass()->GetFName())
-				{
-					// If dragging a class, send though an FAssetData whose name is null and class is this class' name
-					InAssetData.Add(AssetData);
-				}
-				else
-				{
-					InAssetData.Add(AssetData);
-				}
+				// If dragging a class, send though an FAssetData whose name is null and class is this class' name
+				InAssetData.Add(AssetData);
 			}
 			
 			if ( InAssetData.Num() > 0 )
@@ -3587,7 +3580,7 @@ float SAssetView::GetTileViewItemWidth() const
 	return GetTileViewItemBaseWidth() * FillScale;
 }
 
-float SAssetView::GetTileViewItemBaseWidth() const
+float SAssetView::GetTileViewItemBaseWidth() const //-V524
 {
 	return ( TileViewThumbnailSize + TileViewThumbnailPadding * 2 ) * FMath::Lerp( MinThumbnailScale, MaxThumbnailScale, GetThumbnailScale() );
 }
diff --git a/Engine/Source/Editor/LandscapeEditor/Private/LandscapeEdModePaintTools.cpp b/Engine/Source/Editor/LandscapeEditor/Private/LandscapeEdModePaintTools.cpp
--- a/Engine/Source/Editor/LandscapeEditor/Private/LandscapeEdModePaintTools.cpp
+++ b/Engine/Source/Editor/LandscapeEditor/Private/LandscapeEdModePaintTools.cpp
@@ -270,14 +270,7 @@ public:
 
 				if (bUseWeightTargetValue)
 				{
-					if (bInvert)
-					{
-						CurrentValue = FMath::Lerp(CurrentValue, DestValue, PaintAmount / AdjustedStrength);
-					}
-					else
-					{
-						CurrentValue = FMath::Lerp(CurrentValue, DestValue, PaintAmount / AdjustedStrength);
-					}
+					CurrentValue = FMath::Lerp(CurrentValue, DestValue, PaintAmount / AdjustedStrength);
 				}
 				else if (bUseClayBrush)
 				{
diff --git a/Engine/Source/Editor/UMGEditor/Private/Designer/SRuler.cpp b/Engine/Source/Editor/UMGEditor/Private/Designer/SRuler.cpp
--- a/Engine/Source/Editor/UMGEditor/Private/Designer/SRuler.cpp
+++ b/Engine/Source/Editor/UMGEditor/Private/Designer/SRuler.cpp
@@ -256,7 +256,7 @@ int32 SRuler::DrawTicks( FSlateWindowElementList& OutDrawElements, const struct
 	// Draw line that runs along the bottom of all the ticks.
 	{
 		LinePoints[0] = Orientation == Orient_Horizontal ? FVector2D(0, InArgs.AllottedGeometry.Size.Y) : FVector2D(InArgs.AllottedGeometry.Size.X, 0);
-		LinePoints[1] = Orientation == Orient_Horizontal ? FVector2D(InArgs.AllottedGeometry.Size.X, InArgs.AllottedGeometry.Size.Y) : FVector2D(InArgs.AllottedGeometry.Size.X, InArgs.AllottedGeometry.Size.Y);
+		LinePoints[1] = FVector2D(InArgs.AllottedGeometry.Size.X, InArgs.AllottedGeometry.Size.Y);
 
 		const bool bAntiAlias = false;
 		// Draw each sub mark
diff --git a/Engine/Source/Editor/UnrealEd/Private/Commandlets/InternationalizationExportCommandlet.cpp b/Engine/Source/Editor/UnrealEd/Private/Commandlets/InternationalizationExportCommandlet.cpp
--- a/Engine/Source/Editor/UnrealEd/Private/Commandlets/InternationalizationExportCommandlet.cpp
+++ b/Engine/Source/Editor/UnrealEd/Private/Commandlets/InternationalizationExportCommandlet.cpp
@@ -660,8 +660,11 @@ bool FPortableObjectFormatDOM::FromString( const FString& InStr )
 			}
 			for( uint32 NextLineIdx = LineIdx + 1; NextLineIdx < NumFileLines && LinesToProcess[NextLineIdx].Trim().TrimTrailing().StartsWith(TEXT("\"")); ++NextLineIdx)
 			{
-				const FString& NextLine = 
-				RawMsgCtxt += FindDelimitedString(Line, TEXT("\""), TEXT("\""), RawMsgCtxt);
+				FString Tmp;
+				if (FindDelimitedString(Line, TEXT("\""), TEXT("\""), Tmp))
+				{
+					RawMsgCtxt += Tmp;
+				}
 			}
 
 			// If the following line contains more info for this element we continue to parse it out
diff --git a/Engine/Source/Runtime/Core/Private/Stats/StatsCommand.cpp b/Engine/Source/Runtime/Core/Private/Stats/StatsCommand.cpp
--- a/Engine/Source/Runtime/Core/Private/Stats/StatsCommand.cpp
+++ b/Engine/Source/Runtime/Core/Private/Stats/StatsCommand.cpp
@@ -125,9 +125,7 @@ void DumpHistoryFrame(FStatsThreadState const& StatsData, int64 TargetFrame, flo
 			}
 			UE_LOG(LogStats, Log, TEXT("  %s"), *FStatsUtils::DebugPrint(Meta));
 		}
-	}
-	if (DumpCull == 0.0f)
-	{
+
 		UE_LOG(LogStats, Log, TEXT("Exclusive aggregate stack data---------------"));
 		TArray<FStatMessage> Stats;
 		StatsData.GetExclusiveAggregateStackStats(TargetFrame, Stats);
diff --git a/Engine/Source/Runtime/Core/Public/HAL/MallocBinned.h b/Engine/Source/Runtime/Core/Public/HAL/MallocBinned.h
--- a/Engine/Source/Runtime/Core/Public/HAL/MallocBinned.h
+++ b/Engine/Source/Runtime/Core/Public/HAL/MallocBinned.h
@@ -845,7 +845,7 @@ public:
 		*/
 		HashKeyShift = PoolBitShift+IndirectPoolBitShift;
 		/** Used to mask off the bits that have been used to lookup the indirect table */
-		PoolMask =  ( ( 1 << ( HashKeyShift - PoolBitShift ) ) - 1 );
+		PoolMask =  ( ( 1ull << ( HashKeyShift - PoolBitShift ) ) - 1 );
 		BinnedSizeLimit = PAGE_SIZE_LIMIT/2;
 		BinnedOSTableIndex = BinnedSizeLimit+EXTENED_PAGE_POOL_ALLOCATION_COUNT;
 
diff --git a/Engine/Source/Runtime/Engine/Private/Level.cpp b/Engine/Source/Runtime/Engine/Private/Level.cpp
--- a/Engine/Source/Runtime/Engine/Private/Level.cpp
+++ b/Engine/Source/Runtime/Engine/Private/Level.cpp
@@ -805,7 +805,7 @@ void ULevel::CreateModelComponents()
 					Key.Z				= FMath::FloorToInt(NodeBounds.GetCenter().Z / MODEL_GRID_SIZE_Z);
 				}
 
-				Key.MaskedPolyFlags = Surf.PolyFlags & PF_ModelComponentMask;
+				Key.MaskedPolyFlags = 0;
 
 				// Find an existing node list for the grid cell.
 				TArray<uint16>* ComponentNodes = ModelComponentMap.Find(Key);
@@ -1676,11 +1676,11 @@ void ULevel::OnLevelScriptBlueprintChanged(ULevelScriptBlueprint* InBlueprint)
 		FActorSpawnParameters SpawnInfo;
 		SpawnInfo.OverrideLevel = this;
 		LevelScriptActor = OwningWorld->SpawnActor<ALevelScriptActor>( SpawnClass, SpawnInfo );
-		LevelScriptActor->ClearFlags(RF_Transactional);
-		check(LevelScriptActor->GetOuter() == this);
 
 		if( LevelScriptActor )
 		{
+			LevelScriptActor->ClearFlags(RF_Transactional);
+			check(LevelScriptActor->GetOuter() == this);
 			// Finally, fixup all the bound events to point to their new LSA
 			FBlueprintEditorUtils::FixLevelScriptActorBindings(LevelScriptActor, InBlueprint);
 		}		
diff --git a/Engine/Source/Runtime/Engine/Private/Particles/ParticleVertexFactory.cpp b/Engine/Source/Runtime/Engine/Private/Particles/ParticleVertexFactory.cpp
--- a/Engine/Source/Runtime/Engine/Private/Particles/ParticleVertexFactory.cpp
+++ b/Engine/Source/Runtime/Engine/Private/Particles/ParticleVertexFactory.cpp
@@ -105,7 +105,11 @@ public:
 		Offset += sizeof(float) * 4;
 
 		/** The per-particle dynamic parameter stream */
-		Offset = 0;
+
+		// The -V519 disables a warning from PVS-Studio's static analyzer. It noticed that offset is assigned
+		// twice before being read. It is probably safer to leave the redundant assignments here to reduce
+		// the chance of an error being introduced if this code is modified.
+		Offset = 0;  //-V519
 		Stride = sizeof(float) * 4;
 		Elements.Add(FVertexElement(bInstanced ? 2 : 1, Offset, VET_Float4, 5, Stride, bInstanced));
 		Offset += sizeof(float) * 4;
diff --git a/Engine/Source/Runtime/Online/OnlineSubsystemSteam/Private/VoiceEngineSteam.cpp b/Engine/Source/Runtime/Online/OnlineSubsystemSteam/Private/VoiceEngineSteam.cpp
--- a/Engine/Source/Runtime/Online/OnlineSubsystemSteam/Private/VoiceEngineSteam.cpp
+++ b/Engine/Source/Runtime/Online/OnlineSubsystemSteam/Private/VoiceEngineSteam.cpp
@@ -306,7 +306,6 @@ uint32 FVoiceEngineSteam::SubmitRemoteVoiceData(const FUniqueNetId& RemoteTalker
 		return S_OK;
 	}
 
-	bool bAudioComponentCreated = false;
 	// Generate a streaming wave audio component for voice playback
 	if (QueuedData->AudioComponent == NULL || QueuedData->AudioComponent->IsPendingKill())
 	{
@@ -326,7 +325,7 @@ uint32 FVoiceEngineSteam::SubmitRemoteVoiceData(const FUniqueNetId& RemoteTalker
 	if (QueuedData->AudioComponent != NULL)
 	{
 		USoundWaveStreaming* SoundStreaming = CastChecked<USoundWaveStreaming>(QueuedData->AudioComponent->Sound);
-		if (!bAudioComponentCreated && SoundStreaming->GetAvailableAudioByteCount() == 0)
+		if (SoundStreaming->GetAvailableAudioByteCount() == 0)
 		{
 			UE_LOG(LogVoiceDecode, Log, TEXT("VOIP audio component was starved!"));
 		}
diff --git a/Engine/Source/Runtime/UMG/Private/Components/EditableText.cpp b/Engine/Source/Runtime/UMG/Private/Components/EditableText.cpp
--- a/Engine/Source/Runtime/UMG/Private/Components/EditableText.cpp
+++ b/Engine/Source/Runtime/UMG/Private/Components/EditableText.cpp
@@ -20,7 +20,7 @@ UEditableText::UEditableText(const FObjectInitializer& ObjectInitializer)
 
 	// Grab other defaults from slate arguments.
 	IsReadOnly = Defaults._IsReadOnly.Get();
-	IsPassword = Defaults._IsReadOnly.Get();
+	IsPassword = Defaults._IsPassword.Get();
 	MinimumDesiredWidth = Defaults._MinDesiredWidth.Get();
 	IsCaretMovedWhenGainFocus = Defaults._IsCaretMovedWhenGainFocus.Get();
 	SelectAllTextWhenFocused = Defaults._SelectAllTextWhenFocused.Get();
