Fix for UE-11697
With many particle emitters open, perf tanks. Cascade is ticking even if the editor tab is not visible/active, which can cause all sorts of performance issues. Cascade ticks through engine tick, so the only way to properly avoid this is to check for ticking of the viewport widget; added a flag to it to keep track of ticking in the last frame, then checking and clearing that in the Cascade tick function.

diff --git a/Engine/Source/Editor/Cascade/Private/Cascade.cpp b/Engine/Source/Editor/Cascade/Private/Cascade.cpp
--- a/Engine/Source/Editor/Cascade/Private/Cascade.cpp
+++ b/Engine/Source/Editor/Cascade/Private/Cascade.cpp
@@ -1466,6 +1466,19 @@ void FCascade::AddReferencedObjects(FReferenceCollector& Collector)
 
 void FCascade::Tick(float DeltaTime)
 {
+	// This is a bit of a hack. In order to not tick all open Cascade editors (which tick through engine tick) even when not visible,
+	// the preview viewport keeps track of whether it has been ticked in the last frame. Slate widgets aren't ticked if invisible, so 
+	// this will tell us if we should run simulation in this instance. If it hasn't ticked, we skip ticking Cascade as well and clear
+	// the flag for the next frame
+	if ( !PreviewViewport->HasJustTicked() )
+	{
+		return;
+	}
+
+	PreviewViewport->ClearTickFlag();
+
+
+
 	static const double ResetInterval = 0.5f;
 
 	// Clamp delta time.
diff --git a/Engine/Source/Editor/Cascade/Private/SCascadePreviewViewport.cpp b/Engine/Source/Editor/Cascade/Private/SCascadePreviewViewport.cpp
--- a/Engine/Source/Editor/Cascade/Private/SCascadePreviewViewport.cpp
+++ b/Engine/Source/Editor/Cascade/Private/SCascadePreviewViewport.cpp
@@ -10,6 +10,15 @@
 #include "SDockTab.h"
 
 
+
+void SCascadePreviewViewport::Tick(const FGeometry& AllottedGeometry, const double InCurrentTime, const float InDeltaTime)
+{
+	SEditorViewport::Tick(AllottedGeometry, InCurrentTime, InDeltaTime);
+
+	JustTicked = true;
+}
+
+
 SCascadePreviewViewport::~SCascadePreviewViewport()
 {
 	if (ViewportClient.IsValid())
@@ -79,4 +88,4 @@ void SCascadePreviewViewport::OnFocusViewportToSelection()
 	{
 		ViewportClient->FocusViewportOnBox( Component->Bounds.GetBox() );
 	}
-}
\ No newline at end of file
+}
diff --git a/Engine/Source/Editor/Cascade/Private/SCascadePreviewViewport.h b/Engine/Source/Editor/Cascade/Private/SCascadePreviewViewport.h
--- a/Engine/Source/Editor/Cascade/Private/SCascadePreviewViewport.h
+++ b/Engine/Source/Editor/Cascade/Private/SCascadePreviewViewport.h
@@ -37,7 +37,11 @@ public:
 	TSharedPtr<FSceneViewport> GetViewport() const;
 	TSharedPtr<FCascadeEdPreviewViewportClient> GetViewportClient() const;
 	TSharedPtr<SViewport> GetViewportWidget() const;
+
+	virtual void Tick(const FGeometry& AllottedGeometry, const double InCurrentTime, const float InDeltaTime) override;
 	
+	bool HasJustTicked()	{ return JustTicked; }
+	void ClearTickFlag()	{ JustTicked = false; }
 protected:
 	virtual TSharedRef<FEditorViewportClient> MakeEditorViewportClient() override;
 	virtual TSharedPtr<SWidget> MakeViewportToolbar() override;
@@ -52,4 +56,6 @@ private:
 	
 	/** Level viewport client */
 	TSharedPtr<FCascadeEdPreviewViewportClient> ViewportClient;
+
+	bool JustTicked;
 };
