Fix for UE-11697 - Major performance hit with multiple Particle Systems open and tabbed

This was caused by the cascade viewport widget rendering all the time, even when in a hidden tab.

diff --git a/Engine/Source/Editor/Cascade/Private/SCascadePreviewViewport.cpp b/Engine/Source/Editor/Cascade/Private/SCascadePreviewViewport.cpp
--- a/Engine/Source/Editor/Cascade/Private/SCascadePreviewViewport.cpp
+++ b/Engine/Source/Editor/Cascade/Private/SCascadePreviewViewport.cpp
@@ -41,7 +41,7 @@ void SCascadePreviewViewport::RefreshViewport()
 
 bool SCascadePreviewViewport::IsVisible() const
 {
-	return ViewportWidget.IsValid() && (!ParentTab.IsValid() || ParentTab.Pin()->IsForeground());
+	return ViewportWidget.IsValid() && (!ParentTab.IsValid() || ParentTab.Pin()->IsForeground()) && SEditorViewport::IsVisible();
 }
 
 TSharedPtr<FSceneViewport> SCascadePreviewViewport::GetViewport() const
