OR-1761: Exporting emitters crashes the editor, if the destination particle system has more LODs than the source, and if the source emitter is currently soloing.

The LOD difference means that we may add LODs to the source system before duplicating, then duplicate, then remove the added LODs again. That invalidates the soloing information, which we then try to access for a previously non-existent LOD level, and crash. Now updating the soloing setup after adding new LODs.

diff --git a/Engine/Source/Editor/Cascade/Private/Cascade.cpp b/Engine/Source/Editor/Cascade/Private/Cascade.cpp
--- a/Engine/Source/Editor/Cascade/Private/Cascade.cpp
+++ b/Engine/Source/Editor/Cascade/Private/Cascade.cpp
@@ -3150,6 +3150,9 @@ void FCascade::ExportSelectedEmitter()
 				}
 			}
 
+			ParticleSystem->SetupSoloing();		// we may have changed the number of LODs, so our soloing information could be invalid
+
+
 			if (!DuplicateEmitter(SelectedEmitter, DestPartSys, false))
 			{
 				FText Message = FText::Format( NSLOCTEXT("UnrealEd", "Error_FailedToCopyFormatting", "Failed to copy {0} to {1}"), 
