Cascade: Added a viewport overlay display of 'Completed' when a self-terminating system has finished playing (controlled in the View..View Overlays menu)
[INTEGRATE] Change 2700118 by [REDACTED_PERSON][REDACTED_HANDLE] on 2015/09/21 20:07:56

diff --git a/Engine/Source/Editor/Cascade/Private/Cascade.cpp b/Engine/Source/Editor/Cascade/Private/Cascade.cpp
--- a/Engine/Source/Editor/Cascade/Private/Cascade.cpp
+++ b/Engine/Source/Editor/Cascade/Private/Cascade.cpp
@@ -2133,6 +2133,12 @@ void FCascade::BindCommands()
 		FCanExecuteAction(),
 		FIsActionChecked::CreateSP(this, &FCascade::IsViewParticleMemoryChecked));
 
+	ToolkitCommands->MapAction(
+		Commands.View_SystemCompleted,
+		FExecuteAction::CreateSP(this, &FCascade::OnViewSystemCompleted),
+		FCanExecuteAction(),
+		FIsActionChecked::CreateSP(this, &FCascade::IsViewSystemCompletedChecked));
+
 	ToolkitCommands->MapAction(
 		Commands.ToggleGeometry,
 		FExecuteAction::CreateSP(this, &FCascade::OnViewGeometry),
@@ -3581,6 +3587,16 @@ bool FCascade::IsViewParticleMemoryChecked() const
 	return IsDrawOptionEnabled(FCascadeEdPreviewViewportClient::ParticleMemory);
 }
 
+void FCascade::OnViewSystemCompleted()
+{
+	ToggleDrawOption(FCascadeEdPreviewViewportClient::ParticleSystemCompleted);
+}
+
+bool FCascade::IsViewSystemCompletedChecked() const
+{
+	return IsDrawOptionEnabled(FCascadeEdPreviewViewportClient::ParticleSystemCompleted);
+}
+
 void FCascade::OnViewGeometry()
 {
 	if ((PreviewViewport.IsValid() && PreviewViewport->GetViewportClient().IsValid()) &&
diff --git a/Engine/Source/Editor/Cascade/Private/Cascade.h b/Engine/Source/Editor/Cascade/Private/Cascade.h
--- a/Engine/Source/Editor/Cascade/Private/Cascade.h
+++ b/Engine/Source/Editor/Cascade/Private/Cascade.h
@@ -302,6 +302,8 @@ private:
 	bool IsViewParticleTimesChecked() const;
 	void OnViewParticleMemory();
 	bool IsViewParticleMemoryChecked() const;
+	void OnViewSystemCompleted();
+	bool IsViewSystemCompletedChecked() const;
 	void OnViewGeometry();
 	bool IsViewGeometryChecked() const;
 	void OnViewGeometryProperties();
diff --git a/Engine/Source/Editor/Cascade/Private/CascadeActions.cpp b/Engine/Source/Editor/Cascade/Private/CascadeActions.cpp
--- a/Engine/Source/Editor/Cascade/Private/CascadeActions.cpp
+++ b/Engine/Source/Editor/Cascade/Private/CascadeActions.cpp
@@ -34,6 +34,7 @@ void FCascadeCommands::RegisterCommands()
 	UI_COMMAND(View_ParticleEventCounts, "Particle Event Counts", "Display Particle Event Counts", EUserInterfaceActionType::ToggleButton, FInputChord());
 	UI_COMMAND(View_ParticleTimes, "Particle Times", "Display Particle Times", EUserInterfaceActionType::ToggleButton, FInputChord());
 	UI_COMMAND(View_ParticleMemory, "Particle Memory", "Display Particle Memory", EUserInterfaceActionType::ToggleButton, FInputChord());
+	UI_COMMAND(View_SystemCompleted, "System Completed", "Display 'Completed' when the particle system has finished playing and has not been reset (only for non-looping systems)", EUserInterfaceActionType::ToggleButton, FInputChord());
 	UI_COMMAND(ViewMode_Wireframe, "Wireframe", "Select Wireframe Render Mode", EUserInterfaceActionType::RadioButton, FInputChord());
 	UI_COMMAND(ViewMode_Unlit, "Unlit", "Select Unlit Render Mode", EUserInterfaceActionType::RadioButton, FInputChord());
 	UI_COMMAND(ViewMode_Lit, "Lit", "Select Lit Render Mode", EUserInterfaceActionType::RadioButton, FInputChord());
diff --git a/Engine/Source/Editor/Cascade/Private/CascadeActions.h b/Engine/Source/Editor/Cascade/Private/CascadeActions.h
--- a/Engine/Source/Editor/Cascade/Private/CascadeActions.h
+++ b/Engine/Source/Editor/Cascade/Private/CascadeActions.h
@@ -23,6 +23,7 @@ public:
 	TSharedPtr<FUICommandInfo> View_ParticleEventCounts;
 	TSharedPtr<FUICommandInfo> View_ParticleTimes;
 	TSharedPtr<FUICommandInfo> View_ParticleMemory;
+	TSharedPtr<FUICommandInfo> View_SystemCompleted;
 	TSharedPtr<FUICommandInfo> ToggleGeometry;
 	TSharedPtr<FUICommandInfo> ToggleGeometry_Properties;
 	TSharedPtr<FUICommandInfo> ToggleLocalVectorFields;
diff --git a/Engine/Source/Editor/Cascade/Private/CascadePreviewViewportClient.cpp b/Engine/Source/Editor/Cascade/Private/CascadePreviewViewportClient.cpp
--- a/Engine/Source/Editor/Cascade/Private/CascadePreviewViewportClient.cpp
+++ b/Engine/Source/Editor/Cascade/Private/CascadePreviewViewportClient.cpp
@@ -79,7 +79,7 @@ FCascadeEdPreviewViewportClient::FCascadeEdPreviewViewportClient(TWeakPtr<FCasca
 	WidgetMM = WMM_Translate;
 	bManipulatingVectorField = false;
 
-	DrawFlags = ParticleCounts;
+	DrawFlags = ParticleCounts | ParticleSystemCompleted;
 	
 	bUsingOrbitCamera = true;
 
@@ -268,7 +268,7 @@ void FCascadeEdPreviewViewportClient::Draw(FViewport* Viewport, FCanvas* Canvas)
 
 	EngineShowFlags = SavedEngineShowFlags;
 	FCanvasTextItem TextItem( FVector2D::ZeroVector, FText::GetEmpty(), GEngine->GetTinyFont(), FLinearColor::White );
-	if (GetDrawElement(ParticleCounts) || GetDrawElement(ParticleTimes) || GetDrawElement(ParticleEvents) || GetDrawElement(ParticleMemory))
+	if (GetDrawElement(ParticleCounts) || GetDrawElement(ParticleTimes) || GetDrawElement(ParticleEvents) || GetDrawElement(ParticleMemory) || GetDrawElement(ParticleSystemCompleted))
 	{
 		// 'Up' from the lower left...
 		FString strOutput;
@@ -278,7 +278,7 @@ void FCascadeEdPreviewViewportClient::Draw(FViewport* Viewport, FCanvas* Canvas)
 		UParticleSystemComponent* PartComp = CascadePtr.Pin()->GetParticleSystemComponent();
 
 		int32 iWidth, iHeight;
-				
+
 		if (PartComp->EmitterInstances.Num())
 		{
 			for (int32 i = 0; i < PartComp->EmitterInstances.Num(); i++)
@@ -394,6 +394,20 @@ void FCascadeEdPreviewViewportClient::Draw(FViewport* Viewport, FCanvas* Canvas)
 				Canvas->DrawItem( TextItem, XPosition - iWidth, YPosition - iHeight );				
 			}
 		}
+
+		if (GetDrawElement(ParticleSystemCompleted))
+		{
+			if (PartComp->HasCompleted())
+			{
+				TextItem.SetColor(FLinearColor::White);
+				TextItem.Text = LOCTEXT("SystemCompleted", "Completed");
+				TextItem.bCentreX = true;
+				TextItem.bCentreY = true;
+				Canvas->DrawItem(TextItem, Viewport->GetSizeXY().X * 0.5f, Viewport->GetSizeXY().Y - 10);
+				TextItem.bCentreX = false;
+				TextItem.bCentreY = false;
+			}
+		}
 	}
 
 	//Display a warning message in the preview window if the system has no fixed bounding-box and contains a GPU emitter.
diff --git a/Engine/Source/Editor/Cascade/Private/CascadePreviewViewportClient.h b/Engine/Source/Editor/Cascade/Private/CascadePreviewViewportClient.h
--- a/Engine/Source/Editor/Cascade/Private/CascadePreviewViewportClient.h
+++ b/Engine/Source/Editor/Cascade/Private/CascadePreviewViewportClient.h
@@ -49,7 +49,8 @@ public:
 		Bounds = 0x020,
 		WireSphere = 0x040,
 		OriginAxis = 0x080,
-		Orbit = 0x100
+		Orbit = 0x100,
+		ParticleSystemCompleted = 0x200
 	};
 
 	/** Accessors */
diff --git a/Engine/Source/Editor/Cascade/Private/SCascadePreviewToolbar.cpp b/Engine/Source/Editor/Cascade/Private/SCascadePreviewToolbar.cpp
--- a/Engine/Source/Editor/Cascade/Private/SCascadePreviewToolbar.cpp
+++ b/Engine/Source/Editor/Cascade/Private/SCascadePreviewToolbar.cpp
@@ -69,6 +69,7 @@ TSharedRef<SWidget> SCascadePreviewViewportToolBar::GenerateViewMenu() const
 				Menu.AddMenuEntry(FCascadeCommands::Get().View_ParticleEventCounts);
 				Menu.AddMenuEntry(FCascadeCommands::Get().View_ParticleTimes);
 				Menu.AddMenuEntry(FCascadeCommands::Get().View_ParticleMemory);
+				Menu.AddMenuEntry(FCascadeCommands::Get().View_SystemCompleted);
 			}
 			Menu.EndSection();
 		}
