Allow the Feature Level Preview to work within the Cascade viewport

diff --git a/Engine/Source/Editor/Cascade/Private/SCascadePreviewViewport.cpp b/Engine/Source/Editor/Cascade/Private/SCascadePreviewViewport.cpp
--- a/Engine/Source/Editor/Cascade/Private/SCascadePreviewViewport.cpp
+++ b/Engine/Source/Editor/Cascade/Private/SCascadePreviewViewport.cpp
@@ -21,6 +21,9 @@ void SCascadePreviewViewport::Tick(const FGeometry& AllottedGeometry, const doub
 
 SCascadePreviewViewport::~SCascadePreviewViewport()
 {
+	UEditorEngine* Editor = (UEditorEngine*)GEngine;
+	Editor->OnPreviewFeatureLevelChanged().Remove(PreviewFeatureLevelChangedHandle);
+
 	if (ViewportClient.IsValid())
 	{
 		ViewportClient->Viewport = NULL;
@@ -32,6 +35,30 @@ void SCascadePreviewViewport::Construct(const FArguments& InArgs)
 	CascadePtr = InArgs._Cascade;
 
 	SEditorViewport::Construct( SEditorViewport::FArguments() );
+
+	// Restore last used feature level
+	if (ViewportClient.IsValid())
+	{
+		UWorld* World = ViewportClient->GetPreviewScene().GetWorld();
+		if (World != nullptr)
+		{
+			World->ChangeFeatureLevel(GWorld->FeatureLevel);
+		}
+	}
+
+	// Use a delegate to inform the attached world of feature level changes.
+	UEditorEngine* Editor = (UEditorEngine*)GEngine;
+	PreviewFeatureLevelChangedHandle = Editor->OnPreviewFeatureLevelChanged().AddLambda([this](ERHIFeatureLevel::Type NewFeatureLevel)
+		{
+			if(ViewportClient.IsValid())
+			{
+				UWorld* World = ViewportClient->GetPreviewScene().GetWorld();
+				if (World != nullptr)
+				{
+					World->ChangeFeatureLevel(NewFeatureLevel);
+				}
+			}
+		});
 }
 
 void SCascadePreviewViewport::RefreshViewport()
@@ -80,6 +107,20 @@ TSharedPtr<SWidget> SCascadePreviewViewport::MakeViewportToolbar()
 }
 
 
+void SCascadePreviewViewport::PopulateViewportOverlays(TSharedRef<SOverlay> Overlay)
+{
+	SEditorViewport::PopulateViewportOverlays(Overlay);
+
+	// this widget will display the current viewed feature level
+	Overlay->AddSlot()
+		.VAlign(VAlign_Top)
+		.HAlign(HAlign_Right)
+		.Padding(5.0f)
+		[
+			BuildFeatureLevelWidget()
+		];
+}
+
 void SCascadePreviewViewport::OnFocusViewportToSelection()
 {
 	UParticleSystemComponent* Component = CascadePtr.Pin()->GetParticleSystemComponent();
diff --git a/Engine/Source/Editor/Cascade/Private/SCascadePreviewViewport.h b/Engine/Source/Editor/Cascade/Private/SCascadePreviewViewport.h
--- a/Engine/Source/Editor/Cascade/Private/SCascadePreviewViewport.h
+++ b/Engine/Source/Editor/Cascade/Private/SCascadePreviewViewport.h
@@ -50,6 +50,7 @@ public:
 protected:
 	virtual TSharedRef<FEditorViewportClient> MakeEditorViewportClient() override;
 	virtual TSharedPtr<SWidget> MakeViewportToolbar() override;
+	virtual void PopulateViewportOverlays(TSharedRef<SOverlay> Overlay) override;
 	virtual void OnFocusViewportToSelection() override;
 public:
 	/** The parent tab where this viewport resides */
@@ -62,5 +63,8 @@ private:
 	/** Level viewport client */
 	TSharedPtr<FCascadeEdPreviewViewportClient> ViewportClient;
 
+	/** Handle to the registered OnPreviewFeatureLevelChanged delegate. */
+	FDelegateHandle PreviewFeatureLevelChangedHandle;
+
 	bool JustTicked;
 };
