Mirror 4373989 from Dev-Editor.
Cascade: Copy particle emitter properties when invoking Duplicate Emitter or Export Emitter in the Cascade editor.

diff --git a/Engine/Source/Editor/Cascade/Private/Cascade.cpp b/Engine/Source/Editor/Cascade/Private/Cascade.cpp
--- a/Engine/Source/Editor/Cascade/Private/Cascade.cpp
+++ b/Engine/Source/Editor/Cascade/Private/Cascade.cpp
@@ -2815,6 +2815,17 @@ bool FCascade::DuplicateEmitter(UParticleEmitter* SourceEmitter, UParticleSystem
 		NewEmitter->EmitterEditorColor = FColor::MakeRandomColor();
 		NewEmitter->EmitterEditorColor.A = 255;
 
+		NewEmitter->InitialAllocationCount = SourceEmitter->InitialAllocationCount;
+		NewEmitter->QualityLevelSpawnRateScale = SourceEmitter->QualityLevelSpawnRateScale;
+		NewEmitter->DetailModeBitmask = SourceEmitter->DetailModeBitmask;
+		NewEmitter->DetailModeDisplay = SourceEmitter->DetailModeDisplay;
+		NewEmitter->bDisabledLODsKeepEmitterAlive = SourceEmitter->bDisabledLODsKeepEmitterAlive;
+		NewEmitter->bUseLegacySpawningBehavior = SourceEmitter->bUseLegacySpawningBehavior;
+		NewEmitter->bDisableWhenInsignficant = SourceEmitter->bDisableWhenInsignficant;
+		NewEmitter->SignificanceLevel = SourceEmitter->SignificanceLevel;
+		NewEmitter->EmitterRenderMode = SourceEmitter->EmitterRenderMode;
+		NewEmitter->bCollapsed = false;
+
 		//	'Private' data - not required by the editor
 		UObject*			DupObject;
 		UParticleLODLevel*	SourceLODLevel;
