Update Cascade Emitter canvas client to use the newer InputKey functions.

#rnx

diff --git a/Engine/Source/Editor/Cascade/Private/CascadeEmitterCanvasClient.cpp b/Engine/Source/Editor/Cascade/Private/CascadeEmitterCanvasClient.cpp
--- a/Engine/Source/Editor/Cascade/Private/CascadeEmitterCanvasClient.cpp
+++ b/Engine/Source/Editor/Cascade/Private/CascadeEmitterCanvasClient.cpp
@@ -270,32 +270,32 @@ void FCascadeEmitterCanvasClient::Draw(FViewport* InViewport, FCanvas* Canvas)
 	}
 }
 
-bool FCascadeEmitterCanvasClient::InputKey(FViewport* InViewport, int32 ControllerId, FKey Key, EInputEvent Event, float AmountDepressed, bool Gamepad)
+bool FCascadeEmitterCanvasClient::InputKey(const FInputKeyEventArgs& EventArgs)
 {
 	bool bHandled = false;
 
 	bool bLODIsValid = true;
 	UParticleSystem* ParticleSystem = CascadePtr.Pin()->GetParticleSystem();
-	bool bCtrlDown = InViewport->KeyState(EKeys::LeftControl) || InViewport->KeyState(EKeys::RightControl);
-	bool bShiftDown = InViewport->KeyState(EKeys::LeftShift) || InViewport->KeyState(EKeys::RightShift);
-	bool bAltDown = InViewport->KeyState(EKeys::LeftAlt) || InViewport->KeyState(EKeys::RightAlt);
-	int32 HitX = InViewport->GetMouseX();
-	int32 HitY = InViewport->GetMouseY();
+	bool bCtrlDown = EventArgs.Viewport->KeyState(EKeys::LeftControl) || EventArgs.Viewport->KeyState(EKeys::RightControl);
+	bool bShiftDown = EventArgs.Viewport->KeyState(EKeys::LeftShift) || EventArgs.Viewport->KeyState(EKeys::RightShift);
+	bool bAltDown = EventArgs.Viewport->KeyState(EKeys::LeftAlt) || EventArgs.Viewport->KeyState(EKeys::RightAlt);
+	int32 HitX = EventArgs.Viewport->GetMouseX();
+	int32 HitY = EventArgs.Viewport->GetMouseY();
 	FIntPoint MousePos = FIntPoint(HitX, HitY);
 
-	if (Key == EKeys::LeftMouseButton || Key == EKeys::RightMouseButton)
+	if (EventArgs.Key == EKeys::LeftMouseButton || EventArgs.Key == EKeys::RightMouseButton)
 	{
 		bHandled = true;
 		
-		if (Event == IE_Pressed)
+		if (EventArgs.Event == IE_Pressed)
 		{
-			if (Key == EKeys::LeftMouseButton)
+			if (EventArgs.Key == EKeys::LeftMouseButton)
 			{
 				MousePressPosition = MousePos;
 				bMouseDown = true;
 			}
 
-			HHitProxy*	HitResult = InViewport->GetHitProxy(HitX,HitY);
+			HHitProxy*	HitResult = EventArgs.Viewport->GetHitProxy(HitX,HitY);
 			
 			// Short-term, performing a quick-out
 			bool bHandledHitProxy = true;
@@ -307,14 +307,14 @@ bool FCascadeEmitterCanvasClient::InputKey(FViewport* InViewport, int32 Controll
 					UParticleEmitter* Emitter = ((HCascadeEdEmitterProxy*)HitResult)->Emitter;
 					CascadePtr.Pin()->SetSelectedEmitter(Emitter);
 
-					if (Key == EKeys::RightMouseButton)
+					if (EventArgs.Key == EKeys::RightMouseButton)
 					{
 						OpenEmitterMenu();
 					}
 				}
 				else if (HitResult->IsA(HCascadeEdEmitterEnableProxy::StaticGetType()))
 				{
-					if (bLODIsValid && (ParticleSystem != NULL))
+					if (bLODIsValid && (ParticleSystem != nullptr))
 					{
 						CascadePtr.Pin()->ToggleEnableOnSelectedEmitter(((HCascadeEdDrawModeButtonProxy*)HitResult)->Emitter);
 					}
@@ -379,8 +379,8 @@ bool FCascadeEmitterCanvasClient::InputKey(FViewport* InViewport, int32 Controll
 							if ( FColorArray.Num() > 0 )
 							{
 								// Let go of the mouse lock...
-								InViewport->LockMouseToViewport(false);
-								InViewport->CaptureMouse(false);
+								EventArgs.Viewport->LockMouseToViewport(false);
+								EventArgs.Viewport->CaptureMouse(false);
 
 								FColorPickerArgs PickerArgs;
 								PickerArgs.DisplayGamma = TAttribute<float>::Create( TAttribute<float>::FGetter::CreateUObject(GEngine, &UEngine::GetDisplayGamma) );
@@ -398,7 +398,7 @@ bool FCascadeEmitterCanvasClient::InputKey(FViewport* InViewport, int32 Controll
 
 					CascadePtr.Pin()->SetSelectedModule(Emitter, Module);
 
-					if (Key == EKeys::RightMouseButton)
+					if (EventArgs.Key == EKeys::RightMouseButton)
 					{
 						if (bMouseDragging)// && (CurrentMoveMode != CMMM_None))
 						{
@@ -429,7 +429,7 @@ bool FCascadeEmitterCanvasClient::InputKey(FViewport* InViewport, int32 Controll
 						}
 
 						// Figure out and save the offset from mouse location to top-left of selected module.
-						FVector2D ModuleTopLeft(FindModuleTopLeft(Emitter, Module, InViewport));
+						FVector2D ModuleTopLeft(FindModuleTopLeft(Emitter, Module, EventArgs.Viewport));
 						MouseHoldOffset = (ModuleTopLeft*GetDPIScale()- FVector2D(MousePressPosition)).IntPoint();
 					}
 				}
@@ -505,14 +505,16 @@ bool FCascadeEmitterCanvasClient::InputKey(FViewport* InViewport, int32 Controll
 
 			if (bHandledHitProxy == false)
 			{
-				CascadePtr.Pin()->SetSelectedModule(NULL, NULL);
+				CascadePtr.Pin()->SetSelectedModule(nullptr, nullptr);
 
-				if (Key == EKeys::RightMouseButton)
+				if (EventArgs.Key == EKeys::RightMouseButton)
+				{
 					OpenBackgroundMenu();
+				}
 			}
 
 		}
-		else if (Event == IE_Released)
+		else if (EventArgs.Event == IE_Released)
 		{
 			// If we were dragging a module, find where the mouse currently is, and move module there
 			if ((CurrentMoveMode != MoveMode_None) && bMouseDragging)
@@ -536,7 +538,7 @@ bool FCascadeEmitterCanvasClient::InputKey(FViewport* InViewport, int32 Controll
 						if (bDrawDraggedModule)
 						{
 							ModuleDumpList.Add(DraggedModule);
-							DraggedModule = NULL;
+							DraggedModule = nullptr;
 						}
 						else if (CurrentMoveMode == MoveMode_Move)
 						{
@@ -587,15 +589,15 @@ bool FCascadeEmitterCanvasClient::InputKey(FViewport* InViewport, int32 Controll
 			bMouseDown = false;
 			bMouseDragging = false;
 			CurrentMoveMode = MoveMode_None;
-			DraggedModule = NULL;
+			DraggedModule = nullptr;
 
-			InViewport->Invalidate();
+			EventArgs.Viewport->Invalidate();
 		}
-		else if (Event == IE_DoubleClick)
+		else if (EventArgs.Event == IE_DoubleClick)
 		{
-			if (Key == EKeys::LeftMouseButton)
+			if (EventArgs.Key == EKeys::LeftMouseButton)
 			{
-				HHitProxy*	HitResult = InViewport->GetHitProxy(HitX,HitY);
+				HHitProxy*	HitResult = EventArgs.Viewport->GetHitProxy(HitX,HitY);
 				if (HitResult)
 				{
 					if (HitResult->IsA(HCascadeEdEmitterProxy::StaticGetType()))
@@ -608,7 +610,7 @@ bool FCascadeEmitterCanvasClient::InputKey(FViewport* InViewport, int32 Controll
 							{
 								CascadePtr.Pin()->SetSelectedModule(NULL);
 							}
-							InViewport->Invalidate();
+							EventArgs.Viewport->Invalidate();
 						}
 					}
 				}
@@ -616,7 +618,7 @@ bool FCascadeEmitterCanvasClient::InputKey(FViewport* InViewport, int32 Controll
 		}
 	}
 
-	if (Event == IE_Pressed)
+	if (EventArgs.Event == IE_Pressed)
 	{
 		if (bMouseDragging && (CurrentMoveMode != MoveMode_None))
 		{
@@ -625,7 +627,7 @@ bool FCascadeEmitterCanvasClient::InputKey(FViewport* InViewport, int32 Controll
 		}
 		else
 		{
-			if ( Key == EKeys::Platform_Delete )
+			if ( EventArgs.Key == EKeys::Platform_Delete )
 			{
 				if (CascadePtr.Pin()->GetSelectedModule())
 				{
@@ -637,32 +639,32 @@ bool FCascadeEmitterCanvasClient::InputKey(FViewport* InViewport, int32 Controll
 				}
 				bHandled = true;
 			}
-			else if (Key == EKeys::Left)
+			else if (EventArgs.Key == EKeys::Left)
 			{
 				CascadePtr.Pin()->MoveSelectedEmitter(-1);
 				bHandled = true;
 			}
-			else if (Key == EKeys::Right)
+			else if (EventArgs.Key == EKeys::Right)
 			{
 				CascadePtr.Pin()->MoveSelectedEmitter(1);
 				bHandled = true;
 			}
-			else if ((Key == EKeys::Z) && bCtrlDown)
+			else if ((EventArgs.Key == EKeys::Z) && bCtrlDown)
 			{
 				CascadePtr.Pin()->OnUndo();
 				bHandled = true;
 			}
-			else if ((Key == EKeys::Y) && bCtrlDown)
+			else if ((EventArgs.Key == EKeys::Y) && bCtrlDown)
 			{
 				CascadePtr.Pin()->OnRedo();
 				bHandled = true;
 			}
-			else if (Key == EKeys::PageDown)
+			else if (EventArgs.Key == EKeys::PageDown)
 			{
 				CascadePtr.Pin()->OnJumpToLowerLOD();
 				bHandled = true;
 			}
-			else if (Key == EKeys::PageUp)
+			else if (EventArgs.Key == EKeys::PageUp)
 			{
 				CascadePtr.Pin()->OnJumpToHigherLOD();
 				bHandled = true;
@@ -672,7 +674,7 @@ bool FCascadeEmitterCanvasClient::InputKey(FViewport* InViewport, int32 Controll
 
 
 	// Handle viewport screenshot.
-	bHandled |= InputTakeScreenshot(InViewport, Key, Event);
+	bHandled |= InputTakeScreenshot(EventArgs.Viewport, EventArgs.Key, EventArgs.Event);
 
 	return bHandled;
 }
diff --git a/Engine/Source/Editor/Cascade/Private/CascadeEmitterCanvasClient.h b/Engine/Source/Editor/Cascade/Private/CascadeEmitterCanvasClient.h
--- a/Engine/Source/Editor/Cascade/Private/CascadeEmitterCanvasClient.h
+++ b/Engine/Source/Editor/Cascade/Private/CascadeEmitterCanvasClient.h
@@ -27,7 +27,7 @@ public:
 
 	/** FViewportClient interface */
 	virtual void Draw(FViewport* Viewport, FCanvas* Canvas) override;
-	virtual bool InputKey(FViewport* Viewport, int32 ControllerId, FKey Key, EInputEvent Event, float AmountDepressed = 1.0f, bool bGamepad = false) override;
+	virtual bool InputKey(const FInputKeyEventArgs& EventArgs) override;
 	virtual void CapturedMouseMove(FViewport* Viewport, int32 X, int32 Y) override;
 
 	/** Returns the ratio of the size of the particle emitters to the size of the viewport */
