Fixing cascades editor viewport client to handle some changes from main that bypassed all cacade input handling....

#rnx

diff --git a/Engine/Source/Editor/Cascade/Private/CascadeEmitterCanvasClient.cpp b/Engine/Source/Editor/Cascade/Private/CascadeEmitterCanvasClient.cpp
--- a/Engine/Source/Editor/Cascade/Private/CascadeEmitterCanvasClient.cpp
+++ b/Engine/Source/Editor/Cascade/Private/CascadeEmitterCanvasClient.cpp
@@ -271,6 +271,10 @@ void FCascadeEmitterCanvasClient::Draw(FViewport* InViewport, FCanvas* Canvas)
 }
 
 bool FCascadeEmitterCanvasClient::InputKey(const FInputKeyEventArgs& EventArgs)
+{
+	return InputKey(EventArgs.Viewport, EventArgs.ControllerId, EventArgs.Key, EventArgs.Event, EventArgs.AmountDepressed, EventArgs.IsGamepad());
+}
+
 {
 	bool bHandled = false;
 
