Update Cascade Emitter canvas client to use the newer InputKey functions.

#rnx

diff --git a/Engine/Source/Editor/Cascade/Private/CascadeEmitterCanvasClient.cpp b/Engine/Source/Editor/Cascade/Private/CascadeEmitterCanvasClient.cpp
--- a/Engine/Source/Editor/Cascade/Private/CascadeEmitterCanvasClient.cpp
+++ b/Engine/Source/Editor/Cascade/Private/CascadeEmitterCanvasClient.cpp
@@ -271,10 +271,6 @@ void FCascadeEmitterCanvasClient::Draw(FViewport* InViewport, FCanvas* Canvas)
 }
 
 bool FCascadeEmitterCanvasClient::InputKey(const FInputKeyEventArgs& EventArgs)
-{
-	return InputKey(EventArgs.Viewport, EventArgs.ControllerId, EventArgs.Key, EventArgs.Event, EventArgs.AmountDepressed, EventArgs.IsGamepad());
-}
-
 {
 	bool bHandled = false;
 
@@ -514,7 +510,7 @@ bool FCascadeEmitterCanvasClient::InputKey(const FInputKeyEventArgs& EventArgs)
 				if (EventArgs.Key == EKeys::RightMouseButton)
 				{
 					OpenBackgroundMenu();
-				}
+			}
 			}
 
 		}
