Adding additional includes in preparation for removing a global include

#rnx

diff --git a/Engine/Source/Editor/Cascade/Private/Cascade.cpp b/Engine/Source/Editor/Cascade/Private/Cascade.cpp
--- a/Engine/Source/Editor/Cascade/Private/Cascade.cpp
+++ b/Engine/Source/Editor/Cascade/Private/Cascade.cpp
@@ -7,6 +7,7 @@
 #include "Engine/Engine.h"
 #include "Components/StaticMeshComponent.h"
 #include "Editor.h"
+#include "Editor/Transactor.h"
 #include "Components/VectorFieldComponent.h"
 #include "Engine/InterpCurveEdSetup.h"
 #include "CascadeConfiguration.h"
