- Transition thumbnail after rendering

#rnx

diff --git a/Engine/Source/Editor/Cascade/Cascade.Build.cs b/Engine/Source/Editor/Cascade/Cascade.Build.cs
--- a/Engine/Source/Editor/Cascade/Cascade.Build.cs
+++ b/Engine/Source/Editor/Cascade/Cascade.Build.cs
@@ -26,6 +26,7 @@ public class Cascade : ModuleRules
 				"EditorFramework",
 				"UnrealEd",
 				"RHI",
+				"RenderCore",
 				"PhysicsCore"
 			}
 		);
diff --git a/Engine/Source/Editor/Cascade/Private/CascadeEmitterCanvasClient.cpp b/Engine/Source/Editor/Cascade/Private/CascadeEmitterCanvasClient.cpp
--- a/Engine/Source/Editor/Cascade/Private/CascadeEmitterCanvasClient.cpp
+++ b/Engine/Source/Editor/Cascade/Private/CascadeEmitterCanvasClient.cpp
@@ -1160,9 +1160,19 @@ void FCascadeEmitterCanvasClient::DrawHeaderBlock(int32 Index, int32 XPos, UPart
 						if (bNeedsRerender)
 						{
 							Thumbnail.Material = MaterialInterface;
+
 							// Re-render
-							FCanvas ThumbnailCanvas(Thumbnail.Texture->GameThread_GetRenderTargetResource(), nullptr, GetWorld(), GetWorld()->FeatureLevel, FCanvas::CDM_DeferDrawing, ShouldDPIScaleSceneCanvas() ? GetDPIScale() : 1.0f);
-							RenderInfo->Renderer->Draw(MaterialInterface, 0, 0, ScaledSize, ScaledSize, Thumbnail.Texture->GameThread_GetRenderTargetResource(), &ThumbnailCanvas, false);
+							FTextureRenderTargetResource* ThumbnailRTResource = Thumbnail.Texture->GameThread_GetRenderTargetResource();
+							FCanvas ThumbnailCanvas(ThumbnailRTResource, nullptr, GetWorld(), GetWorld()->FeatureLevel, FCanvas::CDM_DeferDrawing, ShouldDPIScaleSceneCanvas() ? GetDPIScale() : 1.0f);
+							RenderInfo->Renderer->Draw(MaterialInterface, 0, 0, ScaledSize, ScaledSize, ThumbnailRTResource, &ThumbnailCanvas, false);
+
+							ENQUEUE_RENDER_COMMAND(TransitionThumbnail_RT)
+							(
+								[RenderTargetResource=ThumbnailRTResource](FRHICommandListImmediate& RHICmdList)
+								{
+									RHICmdList.Transition(FRHITransitionInfo(RenderTargetResource->GetRenderTargetTexture(), ERHIAccess::RTV, ERHIAccess::SRVMask));
+								}
+							);
 						}
 					
 						Canvas->DrawTile(ThumbPos.X - Origin2D.X, ThumbPos.Y - Origin2D.Y, ThumbSize, ThumbSize, 0.f, 0.f, 1.f, 1.f, FLinearColor::White, Thumbnail.Texture->GameThread_GetRenderTargetResource(), false);
