PR #11366: Fix Clang misleading indentation warning

diff --git a/Engine/Source/Editor/Cascade/Private/CascadeEmitterCanvasClient.cpp b/Engine/Source/Editor/Cascade/Private/CascadeEmitterCanvasClient.cpp
--- a/Engine/Source/Editor/Cascade/Private/CascadeEmitterCanvasClient.cpp
+++ b/Engine/Source/Editor/Cascade/Private/CascadeEmitterCanvasClient.cpp
@@ -1538,7 +1538,10 @@ void FCascadeEmitterCanvasClient::DrawDraggedModule(UParticleModule* Module, FVi
 void FCascadeEmitterCanvasClient::DrawCurveButton(UParticleEmitter* Emitter, UParticleModule* Module, bool bHitTesting, FCanvas* Canvas)
 {
 	if (bHitTesting)
+	{
 		Canvas->SetHitProxy(new HCascadeEdGraphButton(Emitter, Module));
+	}
+
 	int32 YPosition = 2;
 	if (CascadePtr.Pin()->GetEditorOptions()->bCenterCascadeModuleText == true)
 	{
