PR #12459: Fix CascadeEmitterCanvasClient.cpp misleading indentation

#rnx

diff --git a/Engine/Source/Editor/Cascade/Private/CascadeEmitterCanvasClient.cpp b/Engine/Source/Editor/Cascade/Private/CascadeEmitterCanvasClient.cpp
--- a/Engine/Source/Editor/Cascade/Private/CascadeEmitterCanvasClient.cpp
+++ b/Engine/Source/Editor/Cascade/Private/CascadeEmitterCanvasClient.cpp
@@ -1600,8 +1600,11 @@ void FCascadeEmitterCanvasClient::Draw3DDrawButton(UParticleEmitter* Emitter, UP
 void FCascadeEmitterCanvasClient::DrawEnableButton(UParticleEmitter* Emitter, UParticleModule* Module, bool bHitTesting, FCanvas* Canvas)
 {
 	if (bHitTesting)
+	{
 		Canvas->SetHitProxy(new HCascadeEdEnableButtonProxy(Emitter, Module));
- 	if (CascadePtr.Pin()->GetEditorOptions()->bUseSlimCascadeDraw == false)
+	}
+
+	if (CascadePtr.Pin()->GetEditorOptions()->bUseSlimCascadeDraw == false)
 	{
 		if (Module->bEnabled)
 		{
@@ -1628,8 +1631,11 @@ void FCascadeEmitterCanvasClient::DrawEnableButton(UParticleEmitter* Emitter, UP
 			Canvas->DrawTile(EmitterWidth - 40, YPosition, 16, 16, 0.f, 0.f, 1.f, 1.f, FLinearColor::White, GetIconTexture(Icon_ModuleDisabled));
 		}
 	}
+
 	if (bHitTesting)
+	{
 		Canvas->SetHitProxy(NULL);
+	}
 }
 
 void FCascadeEmitterCanvasClient::DrawModuleDump(FViewport* InViewport, FCanvas* Canvas)
